(WIP) New BlueprintNativeCodeGen module, to house the commandlet which will convert Blueprint assets into native classes (for cooked builds).

diff --git a/Engine/Source/Developer/BlueprintNativeCodeGen/BlueprintNativeCodeGen.Build.cs b/Engine/Source/Developer/BlueprintNativeCodeGen/BlueprintNativeCodeGen.Build.cs
new file mode 100644
--- /dev/null
+++ b/Engine/Source/Developer/BlueprintNativeCodeGen/BlueprintNativeCodeGen.Build.cs
@@ -0,0 +1,36 @@
+// Copyright 1998-2015 Epic Games, Inc. All Rights Reserved.
+
+using UnrealBuildTool;
+
+public class BlueprintNativeCodeGen : ModuleRules
+{
+    public BlueprintNativeCodeGen(TargetInfo Target)
+	{
+        PrivateIncludePaths.Add("BlueprintNativeCodeGen/Private");
+
+		PublicDependencyModuleNames.AddRange(
+			new string[] {
+				"Core",
+				"CoreUObject",
+				"Engine",
+			}
+		);
+
+		PrivateDependencyModuleNames.AddRange(
+			new string[] {
+                "DesktopPlatform",
+                "UnrealEd",
+                "InputCore",
+                "SlateCore",
+                "Slate",
+                "EditorStyle"
+			}
+		);
+
+        DynamicallyLoadedModuleNames.AddRange(
+			new string[] {
+				"AssetRegistry",
+			}
+		);
+	}
+}
diff --git a/Engine/Source/Developer/BlueprintNativeCodeGen/Private/BlueprintNativeCodeGenCoordinator.cpp b/Engine/Source/Developer/BlueprintNativeCodeGen/Private/BlueprintNativeCodeGenCoordinator.cpp
new file mode 100644
--- /dev/null
+++ b/Engine/Source/Developer/BlueprintNativeCodeGen/Private/BlueprintNativeCodeGenCoordinator.cpp
@@ -0,0 +1,321 @@
+// Copyright 1998-2015 Epic Games, Inc. All Rights Reserved.
+
+#include "BlueprintNativeCodeGenPCH.h"
+#include "BlueprintNativeCodeGenCoordinator.h"
+#include "NativeCodeGenCommandlineParams.h"
+#include "Engine/Blueprint.h"
+#include "Engine/UserDefinedStruct.h"
+#include "Engine/UserDefinedEnum.h"
+#include "ARFilter.h"
+#include "PackageName.h"
+#include "FileManager.h"
+#include "AssetRegistryModule.h"
+
+/*******************************************************************************
+ * UBlueprintNativeCodeGenConfig
+*******************************************************************************/
+
+//------------------------------------------------------------------------------
+UBlueprintNativeCodeGenConfig::UBlueprintNativeCodeGenConfig(FObjectInitializer const& ObjectInitializer)
+	: Super(ObjectInitializer)
+{
+}
+
+/*******************************************************************************
+ * BlueprintNativeCodeGenCoordinatorImpl declaration
+*******************************************************************************/
+
+namespace BlueprintNativeCodeGenCoordinatorImpl
+{
+	static const UClass* TargetAssetTypes[] = {
+		UBlueprint::StaticClass(),
+		UUserDefinedStruct::StaticClass(),
+		UUserDefinedEnum::StaticClass(),
+	};
+
+	static const int32 IncompatibleBlueprintTypeMask = (1 << BPTYPE_LevelScript) |
+		(1 << BPTYPE_MacroLibrary) | (1 << BPTYPE_Const);
+
+	static const FName BlueprintTypePropertyName("BlueprintType");
+	static const FName ParentClassTagName("ParentClassPackage");
+	static const FName AssetInterfacesTagName("ImplementedInterfaces");
+
+	/**  */
+	static IAssetRegistry& GetAssetRegistry();
+
+	/**  */
+	static void GatherAssetsToConvert(TArray<FAssetData>& ConversionQueueOut);
+
+	/** */
+	static void GatherParentAssets(const FAssetData& TargetAsset, TArray<FAssetData>& AssetsOut);
+
+	/** */
+	static void GatherInterfaceAssets(const FAssetData& TargetAsset, TArray<FAssetData>& AssetsOut);
+}
+
+/*******************************************************************************
+ * FScopedAssetCollector
+*******************************************************************************/
+
+/**  */
+struct FScopedAssetCollector
+{
+public:
+	FScopedAssetCollector(const FARFilter& SharedFilter, const UClass* AssetType, TArray<FAssetData>& AssetsOut)
+		: TargetType(AssetType)
+	{
+		if (FScopedAssetCollector** SpecializedCollector = SpecializedAssetCollectors.Find(TargetType))
+		{
+			(*SpecializedCollector)->GatherAssets(SharedFilter, AssetsOut);
+		}
+		else
+		{
+			GatherAssets(SharedFilter, AssetsOut);
+		}
+	}
+
+protected:
+	FScopedAssetCollector(const UClass* AssetType) : TargetType(AssetType) {}
+
+	/**
+	 * 
+	 * 
+	 * @param  SharedFilter    
+	 * @param  AssetsOut    
+	 */
+	virtual void GatherAssets(const FARFilter& SharedFilter, TArray<FAssetData>& AssetsOut) const;
+
+	const UClass* TargetType;
+	static TMap<const UClass*, FScopedAssetCollector*> SpecializedAssetCollectors;
+};
+TMap<const UClass*, FScopedAssetCollector*> FScopedAssetCollector::SpecializedAssetCollectors;
+
+//------------------------------------------------------------------------------
+void FScopedAssetCollector::GatherAssets(const FARFilter& SharedFilter, TArray<FAssetData>& AssetsOut) const
+{
+	IAssetRegistry& AssetRegistry = BlueprintNativeCodeGenCoordinatorImpl::GetAssetRegistry();
+
+	FARFilter ClassFilter;
+	ClassFilter.Append(SharedFilter);
+	ClassFilter.ClassNames.Add(TargetType->GetFName());
+
+	if (ClassFilter.PackagePaths.Num() > 0)
+	{
+		// we want the union of specified package names and paths, not the 
+		// intersection (how GetAssets() is set up to work), so here we query 
+		// for both separately
+		ClassFilter.PackageNames.Empty();
+		AssetRegistry.GetAssets(ClassFilter, AssetsOut);
+		
+		if (SharedFilter.PackageNames.Num() > 0)
+		{
+			ClassFilter.PackageNames = SharedFilter.PackageNames;
+			ClassFilter.PackageNames.Empty();
+			AssetRegistry.GetAssets(ClassFilter, AssetsOut);
+		}
+	}
+	else
+	{
+		AssetRegistry.GetAssets(ClassFilter, AssetsOut);
+	}
+}
+
+/*******************************************************************************
+ * FSpecializedAssetCollector<>
+*******************************************************************************/
+
+template<class AssetType>
+struct FSpecializedAssetCollector : public FScopedAssetCollector
+{
+public:
+	FSpecializedAssetCollector() : FScopedAssetCollector(AssetType::StaticClass())
+	{
+		SpecializedAssetCollectors.Add(AssetType::StaticClass(), this);
+	}
+
+	~FSpecializedAssetCollector()
+	{
+		SpecializedAssetCollectors.Remove(AssetType::StaticClass());
+	}
+
+	// FScopedAssetCollector interface
+	virtual void GatherAssets(const FARFilter& SharedFilter, TArray<FAssetData>& AssetsOut) const override;
+	// End FScopedAssetCollector interface
+};
+
+//------------------------------------------------------------------------------
+template<>
+void FSpecializedAssetCollector<UBlueprint>::GatherAssets(const FARFilter& SharedFilter, TArray<FAssetData>& AssetsOut) const
+{
+	FARFilter BlueprintFilter;
+	BlueprintFilter.Append(SharedFilter);
+
+	if (UByteProperty* TypeProperty = FindField<UByteProperty>(TargetType, BlueprintNativeCodeGenCoordinatorImpl::BlueprintTypePropertyName))
+	{
+		ensure(TypeProperty->HasAnyPropertyFlags(CPF_AssetRegistrySearchable));
+
+		UEnum* TypeEnum = TypeProperty->Enum;
+		if (ensure(TypeEnum != nullptr))
+		{
+			for (int32 BlueprintType = 0; BlueprintType < TypeEnum->GetMaxEnumValue(); ++BlueprintType)
+			{
+				if (!TypeEnum->IsValidEnumValue(BlueprintType))
+				{
+					continue;
+				}
+
+				// explicitly ignore unsupported Blueprint conversion types
+				if ((BlueprintNativeCodeGenCoordinatorImpl::IncompatibleBlueprintTypeMask & (1 << BlueprintType)) != 0)
+				{
+					continue;
+				}
+				BlueprintFilter.TagsAndValues.Add(TypeProperty->GetFName(), TypeEnum->GetEnumName(BlueprintType));
+			}
+		}
+	}
+
+	const int32 StartingIndex = AssetsOut.Num();
+	FScopedAssetCollector::GatherAssets(BlueprintFilter, AssetsOut);
+	const int32 EndingIndex   = AssetsOut.Num();
+
+	// have to ensure that we convert the entire inheritance chain
+	for (int32 BlueprintIndex = StartingIndex; BlueprintIndex < EndingIndex; ++BlueprintIndex)
+	{
+		BlueprintNativeCodeGenCoordinatorImpl::GatherParentAssets(AssetsOut[BlueprintIndex], AssetsOut);
+		BlueprintNativeCodeGenCoordinatorImpl::GatherInterfaceAssets(AssetsOut[BlueprintIndex], AssetsOut);
+		// @TODO: Collect UserDefinedStructs (maybe enums?)... maybe put off till later, once the BP is loaded?
+	}
+}
+
+/*******************************************************************************
+ * BlueprintNativeCodeGenCoordinatorImpl definition
+*******************************************************************************/
+
+//------------------------------------------------------------------------------
+static IAssetRegistry& BlueprintNativeCodeGenCoordinatorImpl::GetAssetRegistry()
+{
+	FAssetRegistryModule& AssetRegistryModule = FModuleManager::LoadModuleChecked<FAssetRegistryModule>(TEXT("AssetRegistry"));
+	return AssetRegistryModule.Get();
+}
+
+//------------------------------------------------------------------------------
+static void BlueprintNativeCodeGenCoordinatorImpl::GatherAssetsToConvert(TArray<FAssetData>& ConversionQueueOut)
+{
+	IAssetRegistry& AssetRegistry = GetAssetRegistry();
+	// have to make sure that the asset registry is fully populated
+	AssetRegistry.SearchAllAssets(/*bSynchronousSearch =*/true);
+
+	FARFilter AssetFilter;
+	AssetFilter.bRecursivePaths = true;
+
+	const UBlueprintNativeCodeGenConfig* ConfigSettings = GetDefault<UBlueprintNativeCodeGenConfig>();
+	IFileManager& FileManager = IFileManager::Get();
+
+	for (const FString& PackagePath : ConfigSettings->PackagesToAlwaysConvert)
+	{
+		const FString RelativePath = FPackageName::LongPackageNameToFilename(PackagePath);
+		if (FileManager.DirectoryExists(*RelativePath))
+		{
+			AssetFilter.PackagePaths.Add(*PackagePath);
+		}
+		else
+		{
+			AssetFilter.PackageNames.Add(*PackagePath);
+		}
+	}
+
+	// will be utilized if UBlueprint is an entry in TargetAssetTypes (auto registers
+	FSpecializedAssetCollector<UBlueprint> BlueprintAssetCollector;
+
+	const int32 TargetClassCount = sizeof(TargetAssetTypes) / sizeof(TargetAssetTypes[0]);
+	for (int32 ClassIndex = 0; ClassIndex < TargetClassCount; ++ClassIndex)
+	{
+		const UClass* TargetClass = TargetAssetTypes[ClassIndex];
+		FScopedAssetCollector AssetCollector(AssetFilter, TargetClass, ConversionQueueOut);
+	}
+}
+
+//------------------------------------------------------------------------------
+static void BlueprintNativeCodeGenCoordinatorImpl::GatherParentAssets(const FAssetData& TargetAsset, TArray<FAssetData>& AssetsOut)
+{
+	IAssetRegistry& AssetRegistry = GetAssetRegistry();
+
+	if (const FString* ParentClassPkg = TargetAsset.TagsAndValues.Find(ParentClassTagName))
+	{
+		TArray<FAssetData> ParentPackages;
+		AssetRegistry.GetAssetsByPackageName(**ParentClassPkg, ParentPackages);
+
+		// should either be 0 (when the parent is a native class), or 1 (the parent is another asset) 
+		ensure(ParentPackages.Num() <= 1);
+
+		for (const FAssetData& ParentData : ParentPackages)
+		{
+			AssetsOut.AddUnique(ParentData); // keeps us from converting parents that were already flagged for conversion
+			GatherParentAssets(ParentData, AssetsOut);
+			GatherInterfaceAssets(ParentData, AssetsOut);
+		}
+	}
+}
+
+//------------------------------------------------------------------------------
+static void BlueprintNativeCodeGenCoordinatorImpl::GatherInterfaceAssets(const FAssetData& TargetAsset, TArray<FAssetData>& AssetsOut)
+{
+	IAssetRegistry& AssetRegistry = GetAssetRegistry();
+
+	if (const FString* InterfacesListPtr = TargetAsset.TagsAndValues.Find(AssetInterfacesTagName))
+	{
+		static const FString InterfaceClassLabel(TEXT("(Interface="));
+		static const FString BlueprintClassName = UBlueprintGeneratedClass::StaticClass()->GetName();
+		static const FString ClassPathDelimiter = TEXT("'");
+		static const FString PackageDelimiter   = TEXT(".");
+
+		FString InterfacesList = *InterfacesListPtr;
+		do
+		{
+			const int32 SearchIndex = InterfacesList.Find(InterfaceClassLabel);
+			if (SearchIndex != INDEX_NONE)
+			{
+				int32 InterfaceClassIndex = SearchIndex + InterfaceClassLabel.Len();
+				InterfacesList = InterfacesList.Mid(InterfaceClassIndex);
+
+				if (InterfacesList.StartsWith(BlueprintClassName))
+				{
+					InterfaceClassIndex = BlueprintClassName.Len();
+
+					const int32 ClassEndIndex = InterfacesList.Find(ClassPathDelimiter, ESearchCase::IgnoreCase, ESearchDir::FromStart, InterfaceClassIndex+1);
+					const FString ClassPath = InterfacesList.Mid(InterfaceClassIndex+1, ClassEndIndex - InterfaceClassIndex);
+
+					FString ClassPackagePath, ClassName;
+					if (ClassPath.Split(PackageDelimiter, &ClassPackagePath, &ClassName))
+					{
+						TArray<FAssetData> InterfacePackages;
+						AssetRegistry.GetAssetsByPackageName(*ClassPackagePath, InterfacePackages);
+
+						// should find a single asset (since we know this is for a Blueprint generated class)
+						ensure(InterfacePackages.Num() == 1);
+
+						for (const FAssetData& Interface : InterfacePackages)
+						{
+							// AddUnique() keeps us from converting interfaces that were already flagged for conversion
+							AssetsOut.AddUnique(Interface); 
+						}
+					}
+				}
+			}
+			else
+			{
+				break;
+			}
+		} while (!InterfacesList.IsEmpty());
+	}
+}
+
+/*******************************************************************************
+ * FBlueprintNativeCodeGenCoordinator
+*******************************************************************************/
+
+//------------------------------------------------------------------------------
+FBlueprintNativeCodeGenCoordinator::FBlueprintNativeCodeGenCoordinator(const FNativeCodeGenCommandlineParams& CommandlineParams)
+{
+	BlueprintNativeCodeGenCoordinatorImpl::GatherAssetsToConvert(ConversionQueue);
+}
diff --git a/Engine/Source/Developer/BlueprintNativeCodeGen/Private/BlueprintNativeCodeGenCoordinator.h b/Engine/Source/Developer/BlueprintNativeCodeGen/Private/BlueprintNativeCodeGenCoordinator.h
new file mode 100644
--- /dev/null
+++ b/Engine/Source/Developer/BlueprintNativeCodeGen/Private/BlueprintNativeCodeGenCoordinator.h
@@ -0,0 +1,32 @@
+// Copyright 1998-2015 Epic Games, Inc. All Rights Reserved.
+
+#pragma once
+
+#include "AssetData.h"
+#include "BlueprintNativeCodeGenManifest.h"
+#include "BlueprintNativeCodeGenCoordinator.generated.h"
+
+// Forward declarations
+struct FNativeCodeGenCommandlineParams;
+
+UCLASS(config=Editor)
+class UBlueprintNativeCodeGenConfig : public UObject
+{
+	GENERATED_UCLASS_BODY()
+
+public:
+	UPROPERTY(globalconfig)
+	TArray<FString> PackagesToAlwaysConvert;
+};
+
+/**  */
+struct FBlueprintNativeCodeGenCoordinator
+{
+public:
+	FBlueprintNativeCodeGenCoordinator(const FNativeCodeGenCommandlineParams& CommandlineParams);
+
+public:
+	TArray<FAssetData> ConversionQueue;
+	FBlueprintNativeCodeGenManifest Manifest;
+};
+
diff --git a/Engine/Source/Developer/BlueprintNativeCodeGen/Private/BlueprintNativeCodeGenManifest.cpp b/Engine/Source/Developer/BlueprintNativeCodeGen/Private/BlueprintNativeCodeGenManifest.cpp
new file mode 100644
--- /dev/null
+++ b/Engine/Source/Developer/BlueprintNativeCodeGen/Private/BlueprintNativeCodeGenManifest.cpp
@@ -0,0 +1,5 @@
+// Copyright 1998-2015 Epic Games, Inc. All Rights Reserved.
+
+#include "BlueprintNativeCodeGenPCH.h"
+#include "BlueprintNativeCodeGenManifest.h"
+
diff --git a/Engine/Source/Developer/BlueprintNativeCodeGen/Private/BlueprintNativeCodeGenManifest.h b/Engine/Source/Developer/BlueprintNativeCodeGen/Private/BlueprintNativeCodeGenManifest.h
new file mode 100644
--- /dev/null
+++ b/Engine/Source/Developer/BlueprintNativeCodeGen/Private/BlueprintNativeCodeGenManifest.h
@@ -0,0 +1,19 @@
+// Copyright 1998-2015 Epic Games, Inc. All Rights Reserved.
+
+#pragma once
+
+/**  */
+struct FBlueprintNativeCodeGenManifest
+{
+public: 
+	struct FAssetEntry
+	{
+		UClass* AssetType;
+		FString AssetPath;
+		FString GeneratedFilePath;
+	};
+
+private:
+	TArray<FAssetEntry> Manifest;
+};
+
diff --git a/Engine/Source/Developer/BlueprintNativeCodeGen/Private/BlueprintNativeCodeGenPCH.h b/Engine/Source/Developer/BlueprintNativeCodeGen/Private/BlueprintNativeCodeGenPCH.h
new file mode 100644
--- /dev/null
+++ b/Engine/Source/Developer/BlueprintNativeCodeGen/Private/BlueprintNativeCodeGenPCH.h
@@ -0,0 +1,5 @@
+// Copyright 1998-2015 Epic Games, Inc. All Rights Reserved.
+
+#pragma once
+
+#include "Engine.h"
diff --git a/Engine/Source/Developer/BlueprintNativeCodeGen/Private/GenerateBlueprintCodeModuleCommandlet.cpp b/Engine/Source/Developer/BlueprintNativeCodeGen/Private/GenerateBlueprintCodeModuleCommandlet.cpp
new file mode 100644
--- /dev/null
+++ b/Engine/Source/Developer/BlueprintNativeCodeGen/Private/GenerateBlueprintCodeModuleCommandlet.cpp
@@ -0,0 +1,47 @@
+// Copyright 1998-2015 Epic Games, Inc. All Rights Reserved.
+
+#include "BlueprintNativeCodeGenPCH.h"
+#include "GenerateBlueprintCodeModuleCommandlet.h"
+#include "NativeCodeGenCommandlineParams.h"
+#include "BlueprintNativeCodeGenCoordinator.h"
+
+DEFINE_LOG_CATEGORY_STATIC(LogBlueprintCodeGen, Log, All);
+
+/*******************************************************************************
+ * GenerateBlueprintCodeModuleImpl
+*******************************************************************************/
+
+namespace GenerateBlueprintCodeModuleImpl
+{
+
+}
+
+/*******************************************************************************
+ * UGenerateBlueprintCodeModuleCommandlet
+*******************************************************************************/
+
+//------------------------------------------------------------------------------
+UGenerateBlueprintCodeModuleCommandlet::UGenerateBlueprintCodeModuleCommandlet(FObjectInitializer const& ObjectInitializer)
+	: Super(ObjectInitializer)
+{
+}
+
+//------------------------------------------------------------------------------
+int32 UGenerateBlueprintCodeModuleCommandlet::Main(FString const& Params)
+{
+	TArray<FString> Tokens, Switches;
+	ParseCommandLine(*Params, Tokens, Switches);
+
+	FNativeCodeGenCommandlineParams CommandlineParams(Switches);
+
+	if (CommandlineParams.bHelpRequested)
+	{
+		UE_LOG(LogBlueprintCodeGen, Display, TEXT("%s"), *FNativeCodeGenCommandlineParams::HelpMessage);
+		return 0;
+	}
+
+	FBlueprintNativeCodeGenCoordinator Coordinator(CommandlineParams);
+	
+	return 0;
+}
+
diff --git a/Engine/Source/Developer/BlueprintNativeCodeGen/Private/GenerateBlueprintCodeModuleCommandlet.h b/Engine/Source/Developer/BlueprintNativeCodeGen/Private/GenerateBlueprintCodeModuleCommandlet.h
new file mode 100644
--- /dev/null
+++ b/Engine/Source/Developer/BlueprintNativeCodeGen/Private/GenerateBlueprintCodeModuleCommandlet.h
@@ -0,0 +1,18 @@
+// Copyright 1998-2015 Epic Games, Inc. All Rights Reserved.
+
+#pragma once
+
+#include "Commandlets/Commandlet.h"
+#include "GenerateBlueprintCodeModuleCommandlet.generated.h"
+
+/**  */
+UCLASS()
+class UGenerateBlueprintCodeModuleCommandlet : public UCommandlet
+{
+	GENERATED_UCLASS_BODY()
+
+public:		
+	// Begin UCommandlet Interface
+	virtual int32 Main(FString const& Params) override;
+	// End UCommandlet Interface
+};
diff --git a/Engine/Source/Developer/BlueprintNativeCodeGen/Private/NativeCodeGenCommandlineParams.cpp b/Engine/Source/Developer/BlueprintNativeCodeGen/Private/NativeCodeGenCommandlineParams.cpp
new file mode 100644
--- /dev/null
+++ b/Engine/Source/Developer/BlueprintNativeCodeGen/Private/NativeCodeGenCommandlineParams.cpp
@@ -0,0 +1,72 @@
+// Copyright 1998-2015 Epic Games, Inc. All Rights Reserved.
+// Copyright 1998-2015 Epic Games, Inc. All Rights Reserved.
+
+#include "BlueprintNativeCodeGenPCH.h"
+#include "NativeCodeGenCommandlineParams.h"
+
+DEFINE_LOG_CATEGORY_STATIC(LogNativeCodeGenCommand, Log, All);
+
+/**  */
+const FString FNativeCodeGenCommandlineParams::HelpMessage = TEXT("\n\
+\n\
+-------------------------------------------------------------------------------\n\
+    GenerateBlueprintCodeModule    ::    Converts Blueprint assets into C++    \n\
+-------------------------------------------------------------------------------\n\
+\n\
+::                                                                             \n\
+:: Usage                                                                       \n\
+::                                                                             \n\
+\n\
+    UE4Editor.exe <project> -run=GenerateBlueprintCodeModule [parameters]      \n\
+\n\
+::                                                                             \n\
+:: Switches                                                                    \n\
+::                                                                             \n\
+\n\
+    -whitelist=<>       \n\
+                        \n\
+\n\
+    -blacklist=<>       \n\
+                        \n\
+\n\
+    -output=<>          \n\
+                        \n\
+\n\
+    -help, -h, -?       Display this message and then exit.                    \n\
+\n");
+
+/*******************************************************************************
+ * NativeCodeGenCommandlineParamsImpl
+ ******************************************************************************/
+
+namespace NativeCodeGenCommandlineParamsImpl
+{
+
+}
+ 
+/*******************************************************************************
+ * FNativeCodeGenCommandlineParams
+ ******************************************************************************/
+
+//------------------------------------------------------------------------------
+FNativeCodeGenCommandlineParams::FNativeCodeGenCommandlineParams(const TArray<FString>& CommandlineSwitches)
+	: bHelpRequested(false)
+{
+	for (const FString& Param : CommandlineSwitches)
+	{
+		FString Switch = Param, Value;
+		Param.Split(TEXT("="), &Switch, &Value);
+
+		if (!Switch.Compare(TEXT("h"), ESearchCase::IgnoreCase) ||
+			!Switch.Compare(TEXT("?"), ESearchCase::IgnoreCase) ||
+			!Switch.Compare(TEXT("help"), ESearchCase::IgnoreCase))
+		{
+			bHelpRequested = true;
+		}
+		else
+		{
+			//UE_LOG(LogNativeCodeGenCommand, Warning, TEXT("Unrecognized commandline parameter: %s"), *Switch);
+		}
+	}
+}
+
diff --git a/Engine/Source/Developer/BlueprintNativeCodeGen/Private/NativeCodeGenCommandlineParams.h b/Engine/Source/Developer/BlueprintNativeCodeGen/Private/NativeCodeGenCommandlineParams.h
new file mode 100644
--- /dev/null
+++ b/Engine/Source/Developer/BlueprintNativeCodeGen/Private/NativeCodeGenCommandlineParams.h
@@ -0,0 +1,20 @@
+// Copyright 1998-2015 Epic Games, Inc. All Rights Reserved.
+
+#pragma once
+
+/**  */
+struct FNativeCodeGenCommandlineParams
+{
+public: 
+	static const FString HelpMessage;
+
+public:
+	FNativeCodeGenCommandlineParams(const TArray<FString>& CommandlineSwitches);
+	
+public:
+	bool            bHelpRequested : 1;
+	TArray<FString> WhiteListedAssetPaths;
+	TArray<FString> BlackListedAssetPaths;
+	FString         ModuleOutputDir;
+};
+
diff --git a/Engine/Source/Developer/BlueprintNativeCodeGen/Private/NativeCodeGenerationTool.cpp b/Engine/Source/Developer/BlueprintNativeCodeGen/Private/NativeCodeGenerationTool.cpp
new file mode 100644
--- /dev/null
+++ b/Engine/Source/Developer/BlueprintNativeCodeGen/Private/NativeCodeGenerationTool.cpp
@@ -0,0 +1,475 @@
+// Copyright 1998-2015 Epic Games, Inc. All Rights Reserved.
+
+#include "BlueprintNativeCodeGenPCH.h"
+#include "NativeCodeGenerationTool.h"
+#include "Dialogs.h"
+#include "Engine/BlueprintGeneratedClass.h"
+#include "Engine/UserDefinedEnum.h"
+#include "Engine/UserDefinedStruct.h"
+#include "SourceCodeNavigation.h"
+#include "DesktopPlatformModule.h" // for InvalidateMakefiles()
+//#include "Editor/KismetCompiler/Public/BlueprintCompilerCppBackendInterface.h"
+
+#define LOCTEXT_NAMESPACE "NativeCodeGenerationTool"
+
+//
+//	THE CODE SHOULD BE MOVED TO GAMEPROJECTGENERATION
+//
+
+
+struct FGeneratedCodeData
+{
+	FGeneratedCodeData(UBlueprint& InBlueprint) 
+		: Blueprint(&InBlueprint)
+	{
+		FName GeneratedClassName, SkeletonClassName;
+		InBlueprint.GetBlueprintClassNames(GeneratedClassName, SkeletonClassName);
+		ClassName = GeneratedClassName.ToString();
+
+		GatherUserDefinedDependencies(InBlueprint);
+	}
+
+	FString TypeDependencies;
+	FString ErrorString;
+	FString ClassName;
+	TWeakObjectPtr<UBlueprint> Blueprint;
+	TSet<UObject*> DependentObjects;
+
+	void GatherUserDefinedDependencies(UBlueprint& InBlueprint)
+	{
+		TArray<UObject*> ReferencedObjects;
+		{
+			FReferenceFinder ReferenceFinder(ReferencedObjects, NULL, false, false, false, false);
+
+			{
+				TArray<UObject*> ObjectsToCheck;
+				GetObjectsWithOuter(InBlueprint.GeneratedClass, ObjectsToCheck, true);
+				for (auto Obj : ObjectsToCheck)
+				{
+					if (IsValid(Obj))
+					{
+						ReferenceFinder.FindReferences(Obj);
+					}
+				}
+			}
+
+			if (InBlueprint.GeneratedClass)
+			{
+				ReferenceFinder.FindReferences(InBlueprint.GeneratedClass->GetDefaultObject(false));
+			}
+
+			for (UClass* Class = InBlueprint.GeneratedClass->GetSuperClass(); Class && !Class->HasAnyClassFlags(CLASS_Native); Class = Class->GetSuperClass())
+			{
+				ReferencedObjects.Add(Class);
+			}
+
+			for (auto& ImplementedInterface : InBlueprint.GeneratedClass->Interfaces)
+			{
+				if (ImplementedInterface.Class && !ImplementedInterface.Class->HasAnyClassFlags(CLASS_Native))
+				{
+					ReferencedObjects.Add(ImplementedInterface.Class);
+				}
+			}
+		}
+
+		TypeDependencies.Empty();
+		for (auto Obj : ReferencedObjects)
+		{
+			if (IsValid(Obj) && !Obj->IsIn(InBlueprint.GeneratedClass) && (Obj != InBlueprint.GeneratedClass))
+			{
+				if (Obj->IsA<UBlueprintGeneratedClass>() || Obj->IsA<UUserDefinedEnum>() || Obj->IsA<UUserDefinedStruct>())
+				{
+					TypeDependencies += Obj->GetPathName();
+					TypeDependencies += TEXT("\n");
+					DependentObjects.Add(Obj);
+				}
+			}
+		}
+		DependentObjects.Add(InBlueprint.GeneratedClass);
+
+		if (TypeDependencies.IsEmpty())
+		{
+			TypeDependencies += LOCTEXT("NoNonNativeDependencies", "No non-native dependencies.").ToString();
+		}
+	}
+
+	static FString DefaultHeaderDir()
+	{
+		auto DefaultSourceDir = FPaths::ConvertRelativePathToFull(FPaths::GameSourceDir());
+		return FPaths::Combine(*DefaultSourceDir, FApp::GetGameName(), TEXT("Public"));
+	}
+
+	static FString DefaultSourceDir()
+	{
+		auto DefaultSourceDir = FPaths::ConvertRelativePathToFull(FPaths::GameSourceDir());
+		return FPaths::Combine(*DefaultSourceDir, FApp::GetGameName(), TEXT("Private"));
+	}
+
+	FString HeaderFileName() const
+	{
+		return ClassName + TEXT(".h");
+	}
+
+	FString SourceFileName() const
+	{
+		return ClassName + TEXT(".cpp");
+	}
+
+	bool Save(const FString& HeaderDirPath, const FString& CppDirPath)
+	{
+		if (!Blueprint.IsValid())
+		{
+			ErrorString += LOCTEXT("InvalidBlueprint", "Invalid Blueprint\n").ToString();
+			return false;
+		}
+
+		const int WorkParts = 4 + 2 * DependentObjects.Num();
+		FScopedSlowTask SlowTask(WorkParts, LOCTEXT("GeneratingCppFiles", "Generating C++ files.."));
+		SlowTask.MakeDialog();
+		SlowTask.EnterProgressFrame();
+
+		TArray<FString> CreatedFiles;
+		for(auto Obj : DependentObjects)
+		{
+			TSharedPtr<FString> HeaderSource(new FString());
+			TSharedPtr<FString> CppSource(new FString());
+			FKismetEditorUtilities::GenerateCppCode(Obj, HeaderSource, CppSource);
+			SlowTask.EnterProgressFrame();
+
+			const FString FullHeaderFilename = FPaths::Combine(*HeaderDirPath, *(Obj->GetName() + TEXT(".h")));
+			const bool bHeaderSaved = FFileHelper::SaveStringToFile(*HeaderSource, *FullHeaderFilename);
+			if (!bHeaderSaved)
+			{
+				ErrorString += FString::Printf(*LOCTEXT("HeaderNotSaved", "Header file wasn't saved. Check log for details. %s\n").ToString(), *Obj->GetPathName());
+			}
+			else
+			{
+				CreatedFiles.Add(FullHeaderFilename);
+			}
+
+			SlowTask.EnterProgressFrame();
+			if (!CppSource->IsEmpty())
+			{
+				const FString NewCppFilename = FPaths::Combine(*CppDirPath, *(Obj->GetName() + TEXT(".cpp")));
+				const bool bCppSaved = FFileHelper::SaveStringToFile(*CppSource, *NewCppFilename);
+				if (!bCppSaved)
+				{
+					ErrorString += FString::Printf(*LOCTEXT("CppNotSaved", "Cpp file wasn't saved. Check log for details. %s\n").ToString(), *Obj->GetPathName());
+				}
+				else
+				{
+					CreatedFiles.Add(NewCppFilename);
+				}
+			}
+		}
+
+		SlowTask.EnterProgressFrame();
+
+		bool bGenerateProjectFiles = true;
+		{
+			bool bProjectHadCodeFiles = false;
+			{
+				TArray<FString> OutProjectCodeFilenames;
+				IFileManager::Get().FindFilesRecursive(OutProjectCodeFilenames, *FPaths::GameSourceDir(), TEXT("*.h"), true, false, false);
+				IFileManager::Get().FindFilesRecursive(OutProjectCodeFilenames, *FPaths::GameSourceDir(), TEXT("*.cpp"), true, false, false);
+				bProjectHadCodeFiles = OutProjectCodeFilenames.Num() > 0;
+			}
+
+			TArray<FString> CreatedFilesForExternalAppRead;
+			CreatedFilesForExternalAppRead.Reserve(CreatedFiles.Num());
+			for (const FString& CreatedFile : CreatedFiles)
+			{
+				CreatedFilesForExternalAppRead.Add(IFileManager::Get().ConvertToAbsolutePathForExternalAppForRead(*CreatedFile));
+			}
+
+			// First see if we can avoid a full generation by adding the new files to an already open project
+			if (bProjectHadCodeFiles && FSourceCodeNavigation::AddSourceFiles(CreatedFilesForExternalAppRead))
+			{
+				// We successfully added the new files to the solution, but we still need to run UBT with -gather to update any UBT makefiles
+				if (FDesktopPlatformModule::Get()->InvalidateMakefiles(FPaths::RootDir(), FPaths::GetProjectFilePath(), GWarn))
+				{
+					// We managed the gather, so we can skip running the full generate
+					bGenerateProjectFiles = false;
+				}
+			}
+		}
+
+		SlowTask.EnterProgressFrame();
+
+		bool bProjectFileUpdated = true;
+		if (bGenerateProjectFiles)
+		{
+			// Generate project files if we happen to be using a project file.
+			if (!FDesktopPlatformModule::Get()->GenerateProjectFiles(FPaths::RootDir(), FPaths::GetProjectFilePath(), GWarn))
+			{
+				ErrorString += LOCTEXT("FailedToGenerateProjectFiles", "Failed to generate project files.").ToString();
+				bProjectFileUpdated = false;
+			}
+		}
+
+		return ErrorString.IsEmpty();
+	}
+};
+
+class SSimpleDirectoryPicker : public SCompoundWidget
+{
+public:
+	SLATE_BEGIN_ARGS(SSimpleDirectoryPicker){}
+		SLATE_ARGUMENT(FString, Directory)
+		SLATE_ARGUMENT(FString, File)
+		SLATE_ARGUMENT(FText, Message)
+		SLATE_ATTRIBUTE(bool, IsEnabled)
+	SLATE_END_ARGS()
+
+	FString File;
+	FString Directory;
+	FText Message;
+
+	TSharedPtr<SEditableTextBox> EditableTextBox;
+
+	FString GetFilePath() const
+	{
+		return FPaths::Combine(*Directory, *File);
+	}
+
+	FText GetFilePathText() const
+	{
+		return FText::FromString(GetFilePath());
+	}
+
+	FReply BrowseHeaderDirectory()
+	{
+		PromptUserForDirectory(Directory, Message.ToString(), Directory);
+
+		//workaround for UI problem
+		EditableTextBox->SetText(TAttribute<FText>::Create(TAttribute<FText>::FGetter::CreateSP(this, &SSimpleDirectoryPicker::GetFilePathText)));
+
+		return FReply::Handled();
+	}
+
+	void Construct(const FArguments& InArgs)
+	{
+		Directory = InArgs._Directory;
+		File = InArgs._File;
+		Message = InArgs._Message;
+
+		TSharedPtr<SButton> OpenButton;
+		ChildSlot
+		[
+			SNew(SHorizontalBox)
+			+ SHorizontalBox::Slot()
+			.FillWidth(1.0f)
+			.Padding(4.0f, 0.0f, 0.0f, 0.0f)
+			[
+				SAssignNew(EditableTextBox, SEditableTextBox)
+				.Text(this, &SSimpleDirectoryPicker::GetFilePathText)
+				.IsReadOnly(true)
+			]
+			+ SHorizontalBox::Slot()
+			.AutoWidth()
+			.Padding(4.0f, 0.0f, 0.0f, 0.0f)
+			[
+				SAssignNew(OpenButton, SButton)
+				.ToolTipText(Message)
+				.OnClicked(this, &SSimpleDirectoryPicker::BrowseHeaderDirectory)
+				[
+					SNew(SImage)
+					.Image(FEditorStyle::GetBrush("PropertyWindow.Button_Ellipsis"))
+				]
+			]
+		];
+
+		OpenButton->SetEnabled(InArgs._IsEnabled);
+	}
+};
+
+class SNativeCodeGenerationDialog : public SCompoundWidget
+{
+public:
+	SLATE_BEGIN_ARGS(SNativeCodeGenerationDialog){}
+		SLATE_ARGUMENT(TSharedPtr<SWindow>, ParentWindow)
+		SLATE_ARGUMENT(TSharedPtr<FGeneratedCodeData>, GeneratedCodeData)
+	SLATE_END_ARGS()
+
+private:
+	
+	// Child widgets
+	TSharedPtr<SSimpleDirectoryPicker> HeaderDirectoryBrowser;
+	TSharedPtr<SSimpleDirectoryPicker> SourceDirectoryBrowser;
+	TSharedPtr<SErrorText> ErrorWidget;
+	// 
+	TWeakPtr<SWindow> WeakParentWindow;
+	TSharedPtr<FGeneratedCodeData> GeneratedCodeData;
+	bool bSaved;
+
+	void CloseParentWindow()
+	{
+		auto ParentWindow = WeakParentWindow.Pin();
+		if (ParentWindow.IsValid())
+		{
+			ParentWindow->RequestDestroyWindow();
+		}
+	}
+
+	bool IsEditable() const
+	{
+		return !bSaved && GeneratedCodeData->ErrorString.IsEmpty();
+	}
+
+	FReply OnButtonClicked()
+	{
+		if (IsEditable())
+		{
+			bSaved = GeneratedCodeData->Save(HeaderDirectoryBrowser->Directory, SourceDirectoryBrowser->Directory);
+			ErrorWidget->SetError(GeneratedCodeData->ErrorString);
+		}
+		else
+		{
+			CloseParentWindow();
+		}
+		
+		return FReply::Handled();
+	}
+
+	FText ButtonText() const
+	{
+		return IsEditable() ? LOCTEXT("Generate", "Generate") : LOCTEXT("Close", "Close");
+	}
+
+	FText GetClassName() const
+	{
+		return GeneratedCodeData.IsValid() ? FText::FromString(GeneratedCodeData->ClassName) : FText::GetEmpty();
+	}
+
+public:
+
+	void Construct(const FArguments& InArgs)
+	{
+		GeneratedCodeData = InArgs._GeneratedCodeData;
+		bSaved = false;
+		WeakParentWindow = InArgs._ParentWindow;
+
+		ChildSlot
+		[
+			SNew(SBorder)
+			.Padding(4.f)
+			.BorderImage(FEditorStyle::GetBrush("ToolPanel.GroupBorder"))
+			[
+				SNew(SVerticalBox)
+				+ SVerticalBox::Slot()
+				.Padding(4.0f)
+				.AutoHeight()
+				[
+					SNew(STextBlock)
+					.Text(LOCTEXT("ClassName", "Class Name"))
+				]
+				+ SVerticalBox::Slot()
+				.Padding(4.0f)
+				.AutoHeight()
+				[
+					SNew(STextBlock)
+					.Text(this, &SNativeCodeGenerationDialog::GetClassName)
+				]
+				+ SVerticalBox::Slot()
+				.Padding(4.0f)
+				.AutoHeight()
+				[
+					SNew(STextBlock)
+					.Text(LOCTEXT("HeaderPath", "Header Path"))
+				]
+				+ SVerticalBox::Slot()
+				.Padding(4.0f)
+				.AutoHeight()
+				[
+					SAssignNew(HeaderDirectoryBrowser, SSimpleDirectoryPicker)
+					.Directory(FGeneratedCodeData::DefaultHeaderDir())
+					.File(GeneratedCodeData->HeaderFileName())
+					.Message(LOCTEXT("HeaderDirectory", "Header Directory"))
+					.IsEnabled(this, &SNativeCodeGenerationDialog::IsEditable)
+				]
+				+ SVerticalBox::Slot()
+				.Padding(4.0f)
+				.AutoHeight()
+				[
+					SNew(STextBlock)
+					.Text(LOCTEXT("SourcePath", "Source Path"))
+				]
+				+ SVerticalBox::Slot()
+				.Padding(4.0f)
+				.AutoHeight()
+				[
+					SAssignNew(SourceDirectoryBrowser, SSimpleDirectoryPicker)
+					.Directory(FGeneratedCodeData::DefaultSourceDir())
+					.File(GeneratedCodeData->SourceFileName())
+					.Message(LOCTEXT("SourceDirectory", "Source Directory"))
+					.IsEnabled(this, &SNativeCodeGenerationDialog::IsEditable)
+				]
+				+ SVerticalBox::Slot()
+				.Padding(4.0f)
+				.AutoHeight()
+				[
+					SNew(STextBlock)
+					.Text(LOCTEXT("Dependencies", "Dependencies"))
+				]
+				+ SVerticalBox::Slot()
+				.Padding(4.0f)
+				.AutoHeight()
+				[
+					SNew(SBox)
+					.WidthOverride(360.0f)
+					.HeightOverride(200.0f)
+					[
+						SNew(SMultiLineEditableTextBox)
+						.IsReadOnly(true)
+						.Text(FText::FromString(GeneratedCodeData->TypeDependencies))
+					]
+				]
+				+ SVerticalBox::Slot()
+				.Padding(4.0f)
+				.AutoHeight()
+				[
+					SAssignNew(ErrorWidget, SErrorText)
+				]
+				+ SVerticalBox::Slot()
+				.Padding(4.0f)
+				.AutoHeight()
+				.HAlign(HAlign_Right)
+				[
+					SNew(SButton)
+					.Text(this, &SNativeCodeGenerationDialog::ButtonText)
+					.OnClicked(this, &SNativeCodeGenerationDialog::OnButtonClicked)
+				]
+			]
+		];
+
+		ErrorWidget->SetError(GeneratedCodeData->ErrorString);
+	}
+};
+
+void FNativeCodeGenerationTool::Open(UBlueprint& Blueprint, TSharedRef< class FBlueprintEditor> Editor)
+{
+	TSharedRef<FGeneratedCodeData> GeneratedCodeData(new FGeneratedCodeData(Blueprint));
+
+	TSharedRef<SWindow> PickerWindow = SNew(SWindow)
+		.Title(LOCTEXT("GenerateNativeCode", "Generate Native Code"))
+		.SizingRule(ESizingRule::Autosized)
+		.ClientSize(FVector2D(0.f, 300.f))
+		.SupportsMaximize(false)
+		.SupportsMinimize(false);
+
+	TSharedRef<SNativeCodeGenerationDialog> CodeGenerationDialog = SNew(SNativeCodeGenerationDialog)
+		.ParentWindow(PickerWindow)
+		.GeneratedCodeData(GeneratedCodeData);
+
+	PickerWindow->SetContent(CodeGenerationDialog);
+	GEditor->EditorAddModalWindow(PickerWindow);
+}
+
+bool FNativeCodeGenerationTool::CanGenerate(const UBlueprint& Blueprint)
+{
+	return (Blueprint.Status == EBlueprintStatus::BS_UpToDate)
+		&& (Blueprint.BlueprintType == EBlueprintType::BPTYPE_Normal || Blueprint.BlueprintType == EBlueprintType::BPTYPE_FunctionLibrary)
+		&& (Blueprint.GeneratedClass->GetClass() == UBlueprintGeneratedClass::StaticClass());
+}
+
+#undef LOCTEXT_NAMESPACE
\ No newline at end of file
diff --git a/Engine/Source/Developer/BlueprintNativeCodeGen/Public/NativeCodeGenerationTool.h b/Engine/Source/Developer/BlueprintNativeCodeGen/Public/NativeCodeGenerationTool.h
new file mode 100644
--- /dev/null
+++ b/Engine/Source/Developer/BlueprintNativeCodeGen/Public/NativeCodeGenerationTool.h
@@ -0,0 +1,14 @@
+// Copyright 1998-2015 Epic Games, Inc. All Rights Reserved.
+
+#pragma once
+
+#include "BlueprintEditor.h"
+
+//////////////////////////////////////////////////////////////////////////
+// FNativeCodeGenerationTool
+
+struct FNativeCodeGenerationTool
+{
+	BLUEPRINTNATIVECODEGEN_API static void Open(UBlueprint& Object, TSharedRef< class FBlueprintEditor> Editor);
+	BLUEPRINTNATIVECODEGEN_API static bool CanGenerate(const UBlueprint& Object);
+};
\ No newline at end of file
diff --git a/Engine/Source/Editor/Kismet/Kismet.Build.cs b/Engine/Source/Editor/Kismet/Kismet.Build.cs
--- a/Engine/Source/Editor/Kismet/Kismet.Build.cs
+++ b/Engine/Source/Editor/Kismet/Kismet.Build.cs
@@ -47,6 +47,7 @@ public class Kismet : ModuleRules
                 "DerivedDataCache",
 				"DesktopPlatform",
 				"HotReload",
+                "BlueprintNativeCodeGen"
 			}
 			);
 
diff --git a/Engine/Source/Editor/Kismet/Private/NativeCodeGenerationTool.cpp b/Engine/Source/Editor/Kismet/Private/NativeCodeGenerationTool.cpp
deleted file mode 100644
diff --git a/Engine/Source/Editor/Kismet/Private/NativeCodeGenerationTool.h b/Engine/Source/Editor/Kismet/Private/NativeCodeGenerationTool.h
deleted file mode 100644
