Setting up the new BlueprintNativeCodeGenModule so that it uses the proper allocator.

diff --git a/Engine/Source/Developer/BlueprintNativeCodeGen/Private/BlueprintNativeCodeGenModule.cpp b/Engine/Source/Developer/BlueprintNativeCodeGen/Private/BlueprintNativeCodeGenModule.cpp
new file mode 100644
--- /dev/null
+++ b/Engine/Source/Developer/BlueprintNativeCodeGen/Private/BlueprintNativeCodeGenModule.cpp
@@ -0,0 +1,21 @@
+// Copyright 1998-2015 Epic Games, Inc. All Rights Reserved.
+
+#include "BlueprintNativeCodeGenPCH.h"
+#include "BlueprintNativeCodeGenModule.h"
+
+/**
+ *
+ */
+class FBlueprintNativeCodeGenModule : public IBlueprintNativeCodeGenModuleInterface
+{
+public:
+	virtual void StartupModule() override
+	{
+	}
+
+	virtual void ShutdownModule() override
+	{
+	}
+};
+
+IMPLEMENT_MODULE(FBlueprintNativeCodeGenModule, BlueprintNativeCodeGen);
diff --git a/Engine/Source/Developer/BlueprintNativeCodeGen/Public/BlueprintNativeCodeGenModule.h b/Engine/Source/Developer/BlueprintNativeCodeGen/Public/BlueprintNativeCodeGenModule.h
new file mode 100644
--- /dev/null
+++ b/Engine/Source/Developer/BlueprintNativeCodeGen/Public/BlueprintNativeCodeGenModule.h
@@ -0,0 +1,10 @@
+// Copyright 1998-2015 Epic Games, Inc. All Rights Reserved.
+
+#pragma once
+
+/** 
+ * 
+ */
+class IBlueprintNativeCodeGenModuleInterface : public IModuleInterface
+{
+};
