(WIP) Blueprint code conversion commandlet - now actually converting assets, recording a manifest that reflects the results.

diff --git a/Engine/Source/Developer/BlueprintNativeCodeGen/BlueprintNativeCodeGen.Build.cs b/Engine/Source/Developer/BlueprintNativeCodeGen/BlueprintNativeCodeGen.Build.cs
--- a/Engine/Source/Developer/BlueprintNativeCodeGen/BlueprintNativeCodeGen.Build.cs
+++ b/Engine/Source/Developer/BlueprintNativeCodeGen/BlueprintNativeCodeGen.Build.cs
@@ -23,7 +23,10 @@ public class BlueprintNativeCodeGen : ModuleRules
                 "InputCore",
                 "SlateCore",
                 "Slate",
-                "EditorStyle"
+                "EditorStyle",
+                "KismetCompiler",
+                "Json",
+                "JsonUtilities",
 			}
 		);
 
diff --git a/Engine/Source/Developer/BlueprintNativeCodeGen/Private/BlueprintNativeCodeGenCoordinator.cpp b/Engine/Source/Developer/BlueprintNativeCodeGen/Private/BlueprintNativeCodeGenCoordinator.cpp
--- a/Engine/Source/Developer/BlueprintNativeCodeGen/Private/BlueprintNativeCodeGenCoordinator.cpp
+++ b/Engine/Source/Developer/BlueprintNativeCodeGen/Private/BlueprintNativeCodeGenCoordinator.cpp
@@ -44,7 +44,7 @@ namespace BlueprintNativeCodeGenCoordinatorImpl
 	static IAssetRegistry& GetAssetRegistry();
 
 	/**  */
-	static void GatherAssetsToConvert(TArray<FAssetData>& ConversionQueueOut);
+	static void GatherAssetsToConvert(const FNativeCodeGenCommandlineParams& CommandlineParams, TArray<FAssetData>& ConversionQueueOut);
 
 	/** */
 	static void GatherParentAssets(const FAssetData& TargetAsset, TArray<FAssetData>& AssetsOut);
@@ -110,7 +110,7 @@ void FScopedAssetCollector::GatherAssets(const FARFilter& SharedFilter, TArray<F
 		if (SharedFilter.PackageNames.Num() > 0)
 		{
 			ClassFilter.PackageNames = SharedFilter.PackageNames;
-			ClassFilter.PackageNames.Empty();
+			ClassFilter.PackagePaths.Empty();
 			AssetRegistry.GetAssets(ClassFilter, AssetsOut);
 		}
 	}
@@ -199,30 +199,36 @@ static IAssetRegistry& BlueprintNativeCodeGenCoordinatorImpl::GetAssetRegistry()
 }
 
 //------------------------------------------------------------------------------
-static void BlueprintNativeCodeGenCoordinatorImpl::GatherAssetsToConvert(TArray<FAssetData>& ConversionQueueOut)
+static void BlueprintNativeCodeGenCoordinatorImpl::GatherAssetsToConvert(const FNativeCodeGenCommandlineParams& CommandlineParams, TArray<FAssetData>& ConversionQueueOut)
 {
 	IAssetRegistry& AssetRegistry = GetAssetRegistry();
 	// have to make sure that the asset registry is fully populated
 	AssetRegistry.SearchAllAssets(/*bSynchronousSearch =*/true);
 
 	FARFilter AssetFilter;
-	AssetFilter.bRecursivePaths = true;
+	AssetFilter.bRecursivePaths = true;	
 
-	const UBlueprintNativeCodeGenConfig* ConfigSettings = GetDefault<UBlueprintNativeCodeGenConfig>();
-	IFileManager& FileManager = IFileManager::Get();
-
-	for (const FString& PackagePath : ConfigSettings->PackagesToAlwaysConvert)
+	auto SetFilterPackages = [&AssetFilter](const TArray<FString>& PackageList)
 	{
-		const FString RelativePath = FPackageName::LongPackageNameToFilename(PackagePath);
-		if (FileManager.DirectoryExists(*RelativePath))
-		{
-			AssetFilter.PackagePaths.Add(*PackagePath);
-		}
-		else
+		IFileManager& FileManager = IFileManager::Get();
+
+		for (const FString& PackagePath : PackageList)
 		{
-			AssetFilter.PackageNames.Add(*PackagePath);
+			const FString RelativePath = FPackageName::LongPackageNameToFilename(PackagePath);
+			if (FileManager.DirectoryExists(*RelativePath))
+			{
+				AssetFilter.PackagePaths.AddUnique(*PackagePath);
+			}
+			else
+			{
+				AssetFilter.PackageNames.AddUnique(*PackagePath);
+			}
 		}
-	}
+	};
+
+	const UBlueprintNativeCodeGenConfig* ConfigSettings = GetDefault<UBlueprintNativeCodeGenConfig>();
+	SetFilterPackages(ConfigSettings->PackagesToAlwaysConvert);
+	SetFilterPackages(CommandlineParams.WhiteListedAssetPaths);
 
 	// will be utilized if UBlueprint is an entry in TargetAssetTypes (auto registers
 	FSpecializedAssetCollector<UBlueprint> BlueprintAssetCollector;
@@ -316,6 +322,54 @@ static void BlueprintNativeCodeGenCoordinatorImpl::GatherInterfaceAssets(const F
 
 //------------------------------------------------------------------------------
 FBlueprintNativeCodeGenCoordinator::FBlueprintNativeCodeGenCoordinator(const FNativeCodeGenCommandlineParams& CommandlineParams)
+	: Manifest(CommandlineParams)
+{
+	BlueprintNativeCodeGenCoordinatorImpl::GatherAssetsToConvert(CommandlineParams, ConversionQueue);
+}
+
+//------------------------------------------------------------------------------
+bool FBlueprintNativeCodeGenCoordinator::IsTargetedForConversion(const FString AssetPath)
+{
+	return ensure(false); // @TODO
+}
+
+//------------------------------------------------------------------------------
+bool FBlueprintNativeCodeGenCoordinator::IsTargetedForConversion(const UBlueprint* Blueprint)
 {
-	BlueprintNativeCodeGenCoordinatorImpl::GatherAssetsToConvert(ConversionQueue);
+	UPackage* AssetPackage = Blueprint->GetOutermost();
+	return IsTargetedForConversion(AssetPackage->GetPathName());
 }
+
+//------------------------------------------------------------------------------
+bool FBlueprintNativeCodeGenCoordinator::IsTargetedForConversion(const UClass* Class)
+{
+	if (Class->HasAnyClassFlags(CLASS_CompiledFromBlueprint))
+	{
+		UPackage* AssetPackage = Class->GetOutermost();
+		return IsTargetedForConversion(AssetPackage->GetPathName());
+	}
+	return false;
+}
+
+//------------------------------------------------------------------------------
+bool FBlueprintNativeCodeGenCoordinator::IsTargetedForConversion(const UEnum* Enum)
+{
+	if (const UUserDefinedEnum* EnumAsset = Cast<UUserDefinedEnum>(Enum))
+	{
+		UPackage* AssetPackage = EnumAsset->GetOutermost();
+		return IsTargetedForConversion(AssetPackage->GetPathName());
+	}
+	return false;
+}
+
+//------------------------------------------------------------------------------
+bool FBlueprintNativeCodeGenCoordinator::IsTargetedForConversion(const UStruct* Struct)
+{
+	if (const UUserDefinedStruct* StructAsset = Cast<UUserDefinedStruct>(Struct))
+	{
+		UPackage* AssetPackage = StructAsset->GetOutermost();
+		return IsTargetedForConversion(AssetPackage->GetPathName());
+	}
+	return false;
+}
+
diff --git a/Engine/Source/Developer/BlueprintNativeCodeGen/Private/BlueprintNativeCodeGenCoordinator.h b/Engine/Source/Developer/BlueprintNativeCodeGen/Private/BlueprintNativeCodeGenCoordinator.h
--- a/Engine/Source/Developer/BlueprintNativeCodeGen/Private/BlueprintNativeCodeGenCoordinator.h
+++ b/Engine/Source/Developer/BlueprintNativeCodeGen/Private/BlueprintNativeCodeGenCoordinator.h
@@ -25,6 +25,12 @@ struct FBlueprintNativeCodeGenCoordinator
 public:
 	FBlueprintNativeCodeGenCoordinator(const FNativeCodeGenCommandlineParams& CommandlineParams);
 
+	bool IsTargetedForConversion(const FString AssetPath);
+	bool IsTargetedForConversion(const UBlueprint* Blueprint);
+	bool IsTargetedForConversion(const UClass* Class);
+	bool IsTargetedForConversion(const UEnum* Enum);
+	bool IsTargetedForConversion(const UStruct* Struct);
+
 public:
 	TArray<FAssetData> ConversionQueue;
 	FBlueprintNativeCodeGenManifest Manifest;
diff --git a/Engine/Source/Developer/BlueprintNativeCodeGen/Private/BlueprintNativeCodeGenManifest.cpp b/Engine/Source/Developer/BlueprintNativeCodeGen/Private/BlueprintNativeCodeGenManifest.cpp
--- a/Engine/Source/Developer/BlueprintNativeCodeGen/Private/BlueprintNativeCodeGenManifest.cpp
+++ b/Engine/Source/Developer/BlueprintNativeCodeGen/Private/BlueprintNativeCodeGenManifest.cpp
@@ -2,4 +2,136 @@
 
 #include "BlueprintNativeCodeGenPCH.h"
 #include "BlueprintNativeCodeGenManifest.h"
+#include "NativeCodeGenCommandlineParams.h"
+#include "App.h" // for GetGameName()
+#include "Serialization/JsonReader.h"
+#include "Serialization/JsonSerializer.h"
+#include "JsonObjectConverter.h"
+#include "Serialization/JsonWriter.h"
+
+
+/*******************************************************************************
+ * BlueprintNativeCodeGenManifestImpl
+ ******************************************************************************/
+
+namespace BlueprintNativeCodeGenManifestImpl
+{
+	static const int64 CPF_NoFlags = 0x00;
+	static const FString ManifestFileExt = TEXT(".bpgen.manifest");
+	static const FString CppFileExt	     = TEXT(".cpp");
+	static const FString HeaderFileExt   = TEXT(".h");
+	static const FString HeaderSubDir    = TEXT("Public");
+	static const FString CppSubDir       = TEXT("Private");
+
+	/** */
+	static FString GetManifestFilePath(const FString& ModulePath);
+
+	/**  */
+	static bool LoadManifest(FBlueprintNativeCodeGenManifest* Manifest);
+
+	/**  */
+	static FString GenerateHeaderSavePath(const FString& ModulePath, const FAssetData& Asset);
+	static FString GenerateCppSavePath(const FString& ModulePath, const FAssetData& Asset);
+}
+
+//------------------------------------------------------------------------------
+static FString BlueprintNativeCodeGenManifestImpl::GetManifestFilePath(const FString& ModulePath)
+{
+	return FPaths::Combine(*ModulePath, FApp::GetGameName()) + ManifestFileExt;
+}
+
+//------------------------------------------------------------------------------
+static bool BlueprintNativeCodeGenManifestImpl::LoadManifest(FBlueprintNativeCodeGenManifest* Manifest)
+{
+	const FString ManifestFilePath = BlueprintNativeCodeGenManifestImpl::GetManifestFilePath(Manifest->GetTargetPath());
+
+	FString ManifestStr;
+	if (FFileHelper::LoadFileToString(ManifestStr, *ManifestFilePath))
+	{
+		TSharedRef< TJsonReader<> > JsonReader = TJsonReaderFactory<>::Create(ManifestStr);
+
+		TSharedPtr<FJsonObject> JsonObject;
+		if (FJsonSerializer::Deserialize(JsonReader, JsonObject))
+		{
+			return FJsonObjectConverter::JsonObjectToUStruct<FBlueprintNativeCodeGenManifest>(JsonObject.ToSharedRef(), Manifest,
+				/*CheckFlags =*/CPF_NoFlags, /*SkipFlags =*/CPF_NoFlags);
+		}
+	}
+	return false;
+}
+
+//------------------------------------------------------------------------------
+static FString BlueprintNativeCodeGenManifestImpl::GenerateHeaderSavePath(const FString& ModulePath, const FAssetData& Asset)
+{
+	// @TODO: Coordinate this with the CppBackend module
+	return FPaths::Combine(*FPaths::Combine(*ModulePath, *HeaderSubDir), *Asset.AssetName.ToString()) + HeaderFileExt;
+}
+
+//------------------------------------------------------------------------------
+static FString BlueprintNativeCodeGenManifestImpl::GenerateCppSavePath(const FString& ModulePath, const FAssetData& Asset)
+{
+	return FPaths::Combine(*FPaths::Combine(*ModulePath, *CppSubDir), *Asset.AssetName.ToString()) + CppFileExt;
+}
+
+/*******************************************************************************
+ * FBlueprintNativeCodeGenManifest
+ ******************************************************************************/
+
+FBlueprintNativeCodeGenManifest::FBlueprintNativeCodeGenManifest(const FString TargetPath)
+	: ModulePath(TargetPath)
+{
+	if (ModulePath.IsEmpty())
+	{
+		ModulePath = FPaths::GameIntermediateDir();
+	}
+	// do NOT load from an existing interface, as this is the default 
+	// constructor used by the USTRUCT() system
+}
+
+//------------------------------------------------------------------------------
+FBlueprintNativeCodeGenManifest::FBlueprintNativeCodeGenManifest(const FNativeCodeGenCommandlineParams& CommandlineParams)
+{
+	ModulePath = CommandlineParams.ModuleOutputDir;
+	// incorporate an existing manifest (in case we're only re-converting a 
+	// handful of assets and adding them into an existing module)
+	if (BlueprintNativeCodeGenManifestImpl::LoadManifest(this))
+	{
+		// reset ModulePath afterwards (in case what we've loaded has been moved
+		ModulePath = CommandlineParams.ModuleOutputDir;
+	}
+}
+
+//------------------------------------------------------------------------------
+FConvertedAssetRecord& FBlueprintNativeCodeGenManifest::CreateConversionRecord(const FAssetData& AssetInfo)
+{
+	FConvertedAssetRecord NewConversionRecord;
+	NewConversionRecord.AssetType = AssetInfo.GetClass();
+	NewConversionRecord.AssetPath = AssetInfo.PackagePath.ToString();
+	NewConversionRecord.GeneratedHeaderPath = BlueprintNativeCodeGenManifestImpl::GenerateHeaderSavePath(ModulePath, AssetInfo);
+	NewConversionRecord.GeneratedCppPath = BlueprintNativeCodeGenManifestImpl::GenerateCppSavePath(ModulePath, AssetInfo);
+
+	return ConvertedAssets[ ConvertedAssets.Add(NewConversionRecord) ];
+}
+
+//------------------------------------------------------------------------------
+bool FBlueprintNativeCodeGenManifest::Save() const
+{
+	TSharedRef<FJsonObject> JsonObject = MakeShareable(new FJsonObject());
+
+	if (FJsonObjectConverter::UStructToJsonObject(FBlueprintNativeCodeGenManifest::StaticStruct(), this, JsonObject,
+		/*CheckFlags =*/BlueprintNativeCodeGenManifestImpl::CPF_NoFlags, /*SkipFlags =*/BlueprintNativeCodeGenManifestImpl::CPF_NoFlags))
+	{
+		FString FileContents;
+		TSharedRef< TJsonWriter<> > JsonWriter = TJsonWriterFactory<>::Create(&FileContents);
+
+		if (FJsonSerializer::Serialize(JsonObject, JsonWriter))
+		{
+			JsonWriter->Close();
+
+			const FString ManifestFilePath = BlueprintNativeCodeGenManifestImpl::GetManifestFilePath(GetTargetPath());
+			return FFileHelper::SaveStringToFile(FileContents, *ManifestFilePath);
+		}
+	}
+	return false;
+}
 
diff --git a/Engine/Source/Developer/BlueprintNativeCodeGen/Private/BlueprintNativeCodeGenManifest.h b/Engine/Source/Developer/BlueprintNativeCodeGen/Private/BlueprintNativeCodeGenManifest.h
--- a/Engine/Source/Developer/BlueprintNativeCodeGen/Private/BlueprintNativeCodeGenManifest.h
+++ b/Engine/Source/Developer/BlueprintNativeCodeGen/Private/BlueprintNativeCodeGenManifest.h
@@ -2,18 +2,65 @@
 
 #pragma once
 
-/**  */
+#include "BlueprintNativeCodeGenManifest.generated.h"
+
+// Forward declarations
+struct FNativeCodeGenCommandlineParams;
+class  FAssetData;
+
+/*******************************************************************************
+ * FCodeGenAssetRecord
+ ******************************************************************************/
+
+USTRUCT()
+struct FConvertedAssetRecord
+{
+	GENERATED_USTRUCT_BODY()
+
+public:
+	UPROPERTY()
+	UClass* AssetType;
+
+	UPROPERTY()
+	FString AssetPath;
+
+	UPROPERTY()
+	FString GeneratedHeaderPath;
+
+	UPROPERTY()
+	FString GeneratedCppPath;
+};
+
+/*******************************************************************************
+ * FBlueprintNativeCodeGenManifest
+ ******************************************************************************/
+ 
+USTRUCT()
 struct FBlueprintNativeCodeGenManifest
 {
+	GENERATED_USTRUCT_BODY()
+
 public: 
-	struct FAssetEntry
-	{
-		UClass* AssetType;
-		FString AssetPath;
-		FString GeneratedFilePath;
-	};
+	FBlueprintNativeCodeGenManifest(const FString TargetPath = TEXT(""));
+	FBlueprintNativeCodeGenManifest(const FNativeCodeGenCommandlineParams& CommandlineParams);
+
+	FORCEINLINE const FString& GetTargetPath() const { return ModulePath; }
+
+	/**  */
+	bool Save() const;
+
+	/**  */
+	FConvertedAssetRecord& CreateConversionRecord(const FAssetData& AssetInfo);
 
 private:
-	TArray<FAssetEntry> Manifest;
-};
+	UPROPERTY()
+	FString ModulePath;
 
+	UPROPERTY()
+	TArray<FString> ModuleDependencies;
+
+	UPROPERTY()
+	TArray<FConvertedAssetRecord> ConvertedAssets; 
+	/**  */
+	TMap<FName, int32> RecordLookupTable;
+};
diff --git a/Engine/Source/Developer/BlueprintNativeCodeGen/Private/GenerateBlueprintCodeModuleCommandlet.cpp b/Engine/Source/Developer/BlueprintNativeCodeGen/Private/GenerateBlueprintCodeModuleCommandlet.cpp
--- a/Engine/Source/Developer/BlueprintNativeCodeGen/Private/GenerateBlueprintCodeModuleCommandlet.cpp
+++ b/Engine/Source/Developer/BlueprintNativeCodeGen/Private/GenerateBlueprintCodeModuleCommandlet.cpp
@@ -4,6 +4,7 @@
 #include "GenerateBlueprintCodeModuleCommandlet.h"
 #include "NativeCodeGenCommandlineParams.h"
 #include "BlueprintNativeCodeGenCoordinator.h"
+#include "BlueprintNativeCodeGenUtils.h"
 
 DEFINE_LOG_CATEGORY_STATIC(LogBlueprintCodeGen, Log, All);
 
@@ -41,6 +42,7 @@ int32 UGenerateBlueprintCodeModuleCommandlet::Main(FString const& Params)
 	}
 
 	FBlueprintNativeCodeGenCoordinator Coordinator(CommandlineParams);
+	FBlueprintNativeCodeGenUtils::GenerateCodeModule(Coordinator);
 	
 	return 0;
 }
diff --git a/Engine/Source/Developer/BlueprintNativeCodeGen/Private/NativeCodeGenCommandlineParams.cpp b/Engine/Source/Developer/BlueprintNativeCodeGen/Private/NativeCodeGenCommandlineParams.cpp
--- a/Engine/Source/Developer/BlueprintNativeCodeGen/Private/NativeCodeGenCommandlineParams.cpp
+++ b/Engine/Source/Developer/BlueprintNativeCodeGen/Private/NativeCodeGenCommandlineParams.cpp
@@ -4,13 +4,13 @@
 #include "BlueprintNativeCodeGenPCH.h"
 #include "NativeCodeGenCommandlineParams.h"
 
-DEFINE_LOG_CATEGORY_STATIC(LogNativeCodeGenCommand, Log, All);
+DEFINE_LOG_CATEGORY_STATIC(LogNativeCodeGenCommandline, Log, All);
 
 /**  */
 const FString FNativeCodeGenCommandlineParams::HelpMessage = TEXT("\n\
 \n\
 -------------------------------------------------------------------------------\n\
-    GenerateBlueprintCodeModule    ::    Converts Blueprint assets into C++    \n\
+::  GenerateBlueprintCodeModule    ::    Converts Blueprint assets into C++    \n\
 -------------------------------------------------------------------------------\n\
 \n\
 ::                                                                             \n\
@@ -20,17 +20,24 @@ const FString FNativeCodeGenCommandlineParams::HelpMessage = TEXT("\n\
     UE4Editor.exe <project> -run=GenerateBlueprintCodeModule [parameters]      \n\
 \n\
 ::                                                                             \n\
-:: Switches                                                                    \n\
+:: Parameters                                                                  \n\
 ::                                                                             \n\
 \n\
-    -whitelist=<>       \n\
-                        \n\
+    -whitelist=<Assets> Identifies assets that you wish to convert. <Assets>   \n\
+                        can be a comma seperated list, specifying multiple     \n\
+                        packages (either directories, or explicit assets) in   \n\
+                        the form: /Game/MyContentDir,/Engine/EngineAssetName.  \n\
 \n\
-    -blacklist=<>       \n\
-                        \n\
+    -blacklist=<Assets> Explicitly specifies assets that you don't want        \n\
+                        converted (listed in the same manner as -whitelist).   \n\
+                        This takes priority over the whitelist (even if it     \n\
+                        results in uncompilable code).                         \n\
 \n\
-    -output=<>          \n\
-                        \n\
+    -output=<ModuleDir> Specifies the path where you want converted assets     \n\
+                        saved to. If left unset, this will default to the      \n\
+                        project's intermeadiate folder (within a sub-folder of \n\
+                        its own). If specified as a relative path, it will be  \n\
+                        relative to the target project's root directory.       \n\
 \n\
     -help, -h, -?       Display this message and then exit.                    \n\
 \n");
@@ -41,7 +48,8 @@ const FString FNativeCodeGenCommandlineParams::HelpMessage = TEXT("\n\
 
 namespace NativeCodeGenCommandlineParamsImpl
 {
-
+	static const FString DefaultModuleName(TEXT("GeneratedBlueprintCode"));
+	static const FString ParamListDelim(TEXT(","));
 }
  
 /*******************************************************************************
@@ -52,6 +60,9 @@ namespace NativeCodeGenCommandlineParamsImpl
 FNativeCodeGenCommandlineParams::FNativeCodeGenCommandlineParams(const TArray<FString>& CommandlineSwitches)
 	: bHelpRequested(false)
 {
+	IFileManager& FileManager = IFileManager::Get();
+	ModuleOutputDir = FPaths::Combine(*FPaths::GameIntermediateDir(), *NativeCodeGenCommandlineParamsImpl::DefaultModuleName);
+
 	for (const FString& Param : CommandlineSwitches)
 	{
 		FString Switch = Param, Value;
@@ -63,9 +74,35 @@ FNativeCodeGenCommandlineParams::FNativeCodeGenCommandlineParams(const TArray<FS
 		{
 			bHelpRequested = true;
 		}
+		else if (!Switch.Compare(TEXT("whitelist"), ESearchCase::IgnoreCase)) 
+		{
+			Value.ParseIntoArray(WhiteListedAssetPaths, *NativeCodeGenCommandlineParamsImpl::ParamListDelim);
+		}
+		else if (!Switch.Compare(TEXT("blacklist"), ESearchCase::IgnoreCase))
+		{
+			Value.ParseIntoArray(BlackListedAssetPaths, *NativeCodeGenCommandlineParamsImpl::ParamListDelim);
+		}
+		else if (!Switch.Compare(TEXT("output"), ESearchCase::IgnoreCase))
+		{
+			// if it's a relative path, let's have it relative to the game 
+			// directory (not the UE executable)
+			if (FPaths::IsRelative(*Value))
+			{
+				Value = FPaths::Combine(*FPaths::GameDir(), *Value);
+			}
+
+			if (FileManager.DirectoryExists(*Value) || FileManager.MakeDirectory(*Value))
+			{
+				ModuleOutputDir = Value;
+			}
+			else
+			{
+				UE_LOG(LogNativeCodeGenCommandline, Warning, TEXT("'%s' doesn't appear to be a usable output directory, defaulting to: '%s'"), *Value, *ModuleOutputDir);
+			}
+		}
 		else
 		{
-			//UE_LOG(LogNativeCodeGenCommand, Warning, TEXT("Unrecognized commandline parameter: %s"), *Switch);
+			//UE_LOG(LogNativeCodeGenCommandline, Warning, TEXT("Unrecognized commandline parameter: %s"), *Switch);
 		}
 	}
 }
diff --git a/Engine/Source/Developer/BlueprintNativeCodeGen/Private/NativeCodeGenerationTool.cpp b/Engine/Source/Developer/BlueprintNativeCodeGen/Private/NativeCodeGenerationTool.cpp
--- a/Engine/Source/Developer/BlueprintNativeCodeGen/Private/NativeCodeGenerationTool.cpp
+++ b/Engine/Source/Developer/BlueprintNativeCodeGen/Private/NativeCodeGenerationTool.cpp
@@ -133,7 +133,7 @@ struct FGeneratedCodeData
 		{
 			TSharedPtr<FString> HeaderSource(new FString());
 			TSharedPtr<FString> CppSource(new FString());
-			FKismetEditorUtilities::GenerateCppCode(Obj, HeaderSource, CppSource);
+			FBlueprintNativeCodeGenUtils::GenerateCppCode(Obj, HeaderSource, CppSource);
 			SlowTask.EnterProgressFrame();
 
 			const FString FullHeaderFilename = FPaths::Combine(*HeaderDirPath, *(Obj->GetName() + TEXT(".h")));
diff --git a/Engine/Source/Editor/UnrealEd/Private/Kismet2/Kismet2.cpp b/Engine/Source/Editor/UnrealEd/Private/Kismet2/Kismet2.cpp
--- a/Engine/Source/Editor/UnrealEd/Private/Kismet2/Kismet2.cpp
+++ b/Engine/Source/Editor/UnrealEd/Private/Kismet2/Kismet2.cpp
@@ -957,62 +957,6 @@ void FKismetEditorUtilities::RecompileBlueprintBytecode(UBlueprint* BlueprintObj
 	}
 }
 
-/** Recompiles the bytecode of a blueprint only.  Should only be run for recompiling dependencies during compile on load */
-void FKismetEditorUtilities::GenerateCppCode(UObject* Obj, TSharedPtr<FString> OutHeaderSource, TSharedPtr<FString> OutCppSource)
-{
-	auto UDEnum = Cast<UUserDefinedEnum>(Obj);
-	auto UDStruct = Cast<UUserDefinedStruct>(Obj);
-	auto BPGC = Cast<UClass>(Obj);
-	auto InBlueprintObj = BPGC ? Cast<UBlueprint>(BPGC->ClassGeneratedBy) : nullptr;
-
-	if (InBlueprintObj)
-	{
-		check(InBlueprintObj->GetOutermost() != GetTransientPackage());
-		checkf(InBlueprintObj->GeneratedClass, TEXT("Invalid generated class for %s"), *InBlueprintObj->GetName());
-		check(OutHeaderSource.IsValid());
-		check(OutCppSource.IsValid());
-
-		auto BlueprintObj = InBlueprintObj;
-		{
-			auto Reinstancer = FBlueprintCompileReinstancer::Create(BlueprintObj->GeneratedClass);
-
-			IKismetCompilerInterface& Compiler = FModuleManager::LoadModuleChecked<IKismetCompilerInterface>(KISMET_COMPILER_MODULENAME);
-
-			TGuardValue<bool> GuardTemplateNameFlag(GCompilingBlueprint, true);
-			FCompilerResultsLog Results;
-
-			FKismetCompilerOptions CompileOptions;
-			CompileOptions.CompileType = EKismetCompileType::Cpp;
-			CompileOptions.OutCppSourceCode = OutCppSource;
-			CompileOptions.OutHeaderSourceCode = OutHeaderSource;
-			Compiler.CompileBlueprint(BlueprintObj, CompileOptions, Results);
-
-			if (EBlueprintType::BPTYPE_Interface == BlueprintObj->BlueprintType && OutCppSource.IsValid())
-			{
-				OutCppSource->Empty(); // ugly temp hack
-			}
-		}
-		CompileBlueprint(BlueprintObj);
-	}
-	else if ((UDEnum || UDStruct) && OutHeaderSource.IsValid())
-	{
-		IKismetCompilerInterface& Compiler = FModuleManager::LoadModuleChecked<IKismetCompilerInterface>(KISMET_COMPILER_MODULENAME);
-		if (UDEnum)
-		{
-			*OutHeaderSource = Compiler.GenerateCppCodeForEnum(UDEnum);
-		}
-		else if (UDStruct)
-		{
-			*OutHeaderSource = Compiler.GenerateCppCodeForStruct(UDStruct);
-		}
-	}
-	else
-	{
-		ensure(false);
-	}
-}
-
-
 namespace ConformComponentsUtils
 {
 	static void ConformRemovedNativeComponents(UObject* BpCdo);
diff --git a/Engine/Source/Editor/UnrealEd/Public/Kismet2/KismetEditorUtilities.h b/Engine/Source/Editor/UnrealEd/Public/Kismet2/KismetEditorUtilities.h
--- a/Engine/Source/Editor/UnrealEd/Public/Kismet2/KismetEditorUtilities.h
+++ b/Engine/Source/Editor/UnrealEd/Public/Kismet2/KismetEditorUtilities.h
@@ -105,8 +105,6 @@ public:
 	/** Recompiles the bytecode of a blueprint only.  Should only be run for recompiling dependencies during compile on load */
 	static void RecompileBlueprintBytecode(UBlueprint* BlueprintObj, TArray<UObject*>* ObjLoaded = nullptr, bool bBatchCompile = false);
 
-	static void GenerateCppCode(UObject* Obj, TSharedPtr<FString> OutHeaderSource, TSharedPtr<FString> OutCppSource);
-
 	/** Tries to make sure that a data-only blueprint is conformed to its native parent, in case any native class flags have changed */
 	static void ConformBlueprintFlagsAndComponents(UBlueprint* BlueprintObj);
 
