CIS Fix - Adding missing files.

diff --git a/Engine/Source/Developer/BlueprintNativeCodeGen/Private/BlueprintNativeCodeGenUtils.cpp b/Engine/Source/Developer/BlueprintNativeCodeGen/Private/BlueprintNativeCodeGenUtils.cpp
new file mode 100644
--- /dev/null
+++ b/Engine/Source/Developer/BlueprintNativeCodeGen/Private/BlueprintNativeCodeGenUtils.cpp
@@ -0,0 +1,125 @@
+// Copyright 1998-2015 Epic Games, Inc. All Rights Reserved.
+
+#include "BlueprintNativeCodeGenPCH.h"
+#include "BlueprintNativeCodeGenUtils.h"
+#include "BlueprintNativeCodeGenCoordinator.h"
+#include "Kismet2/KismetReinstanceUtilities.h"	// for FBlueprintCompileReinstancer
+#include "Kismet2/CompilerResultsLog.h" 
+#include "KismetCompilerModule.h"
+#include "Engine/Blueprint.h"
+#include "Engine/UserDefinedStruct.h"
+#include "Engine/UserDefinedEnum.h"
+#include "Kismet2/KismetEditorUtilities.h"		// for CompileBlueprint()
+
+/*******************************************************************************
+ * BlueprintNativeCodeGenUtilsImpl
+ ******************************************************************************/
+
+namespace BlueprintNativeCodeGenUtilsImpl
+{
+	
+}
+
+/*******************************************************************************
+ * FBlueprintNativeCodeGenUtils
+ ******************************************************************************/
+
+//------------------------------------------------------------------------------
+bool FBlueprintNativeCodeGenUtils::GenerateCodeModule(FBlueprintNativeCodeGenCoordinator& Coordinator)
+{
+	
+	TSharedPtr<FString> HeaderSource(new FString());
+	TSharedPtr<FString> CppSource(new FString());
+
+	FBlueprintNativeCodeGenManifest& Manifest = Coordinator.Manifest;
+
+	bool bSuccess = (Coordinator.ConversionQueue.Num() > 0);
+	for (const FAssetData& Asset : Coordinator.ConversionQueue)
+	{
+		FConvertedAssetRecord& ConversionRecord = Manifest.CreateConversionRecord(Asset);
+
+		// loads the object if it is not already loaded
+		UObject* AssetObj = Asset.GetAsset();
+		FBlueprintNativeCodeGenUtils::GenerateCppCode(AssetObj, HeaderSource, CppSource);
+
+		bSuccess &= !HeaderSource->IsEmpty() || !CppSource->IsEmpty();
+		if (!HeaderSource->IsEmpty())
+		{
+			bSuccess &= FFileHelper::SaveStringToFile(*HeaderSource, *ConversionRecord.GeneratedHeaderPath);
+		}
+		else
+		{
+			ConversionRecord.GeneratedHeaderPath.Empty();
+		}
+
+		if (!CppSource->IsEmpty())
+		{
+			bSuccess &= FFileHelper::SaveStringToFile(*CppSource, *ConversionRecord.GeneratedCppPath);
+		}
+		else
+		{
+			ConversionRecord.GeneratedCppPath.Empty();
+		}
+	}
+	Coordinator.Manifest.Save();
+
+	return bSuccess;
+}
+
+//------------------------------------------------------------------------------
+void FBlueprintNativeCodeGenUtils::GenerateCppCode(UObject* Obj, TSharedPtr<FString> OutHeaderSource, TSharedPtr<FString> OutCppSource)
+{
+	auto UDEnum = Cast<UUserDefinedEnum>(Obj);
+	auto UDStruct = Cast<UUserDefinedStruct>(Obj);
+	auto BPGC = Cast<UClass>(Obj);
+	auto InBlueprintObj = BPGC ? Cast<UBlueprint>(BPGC->ClassGeneratedBy) : Cast<UBlueprint>(Obj);
+
+	OutHeaderSource->Empty();
+	OutCppSource->Empty();
+
+	if (InBlueprintObj)
+	{
+		check(InBlueprintObj->GetOutermost() != GetTransientPackage());
+		checkf(InBlueprintObj->GeneratedClass, TEXT("Invalid generated class for %s"), *InBlueprintObj->GetName());
+		check(OutHeaderSource.IsValid());
+		check(OutCppSource.IsValid());
+
+		auto BlueprintObj = InBlueprintObj;
+		{
+			TSharedPtr<FBlueprintCompileReinstancer> Reinstancer = FBlueprintCompileReinstancer::Create(BlueprintObj->GeneratedClass);
+
+			IKismetCompilerInterface& Compiler = FModuleManager::LoadModuleChecked<IKismetCompilerInterface>(KISMET_COMPILER_MODULENAME);
+			
+			TGuardValue<bool> GuardTemplateNameFlag(GCompilingBlueprint, true);
+			FCompilerResultsLog Results;			
+
+			FKismetCompilerOptions CompileOptions;
+			CompileOptions.CompileType = EKismetCompileType::Cpp;
+			CompileOptions.OutCppSourceCode = OutCppSource;
+			CompileOptions.OutHeaderSourceCode = OutHeaderSource;
+			Compiler.CompileBlueprint(BlueprintObj, CompileOptions, Results);
+
+			if (EBlueprintType::BPTYPE_Interface == BlueprintObj->BlueprintType && OutCppSource.IsValid())
+			{
+				OutCppSource->Empty(); // ugly temp hack
+			}
+		}
+		FKismetEditorUtilities::CompileBlueprint(BlueprintObj);
+	}
+	else if ((UDEnum || UDStruct) && OutHeaderSource.IsValid())
+	{
+		IKismetCompilerInterface& Compiler = FModuleManager::LoadModuleChecked<IKismetCompilerInterface>(KISMET_COMPILER_MODULENAME);
+		if (UDEnum)
+		{
+			*OutHeaderSource = Compiler.GenerateCppCodeForEnum(UDEnum);
+		}
+		else if (UDStruct)
+		{
+			*OutHeaderSource = Compiler.GenerateCppCodeForStruct(UDStruct);
+		}
+	}
+	else
+	{
+		ensure(false);
+	}
+}
diff --git a/Engine/Source/Developer/BlueprintNativeCodeGen/Private/BlueprintNativeCodeGenUtils.h b/Engine/Source/Developer/BlueprintNativeCodeGen/Private/BlueprintNativeCodeGenUtils.h
new file mode 100644
--- /dev/null
+++ b/Engine/Source/Developer/BlueprintNativeCodeGen/Private/BlueprintNativeCodeGenUtils.h
@@ -0,0 +1,21 @@
+// Copyright 1998-2015 Epic Games, Inc. All Rights Reserved.
+
+#pragma once
+
+// Forward declares
+struct FBlueprintNativeCodeGenCoordinator;
+
+/**  */
+struct FBlueprintNativeCodeGenUtils
+{
+public: 
+	/**  */
+	static bool GenerateCodeModule(FBlueprintNativeCodeGenCoordinator& Coordinator);
+
+	/** 
+	 * Recompiles the bytecode of a blueprint only. Should only be run for 
+	 * recompiling dependencies during compile on load. 
+	 */
+	static void GenerateCppCode(UObject* Obj, TSharedPtr<FString> OutHeaderSource, TSharedPtr<FString> OutCppSource);
+};
+
