CIS fix

diff --git a/Engine/Source/Developer/BlueprintNativeCodeGen/Private/GenerateBlueprintCodeModuleCommandlet.cpp b/Engine/Source/Developer/BlueprintNativeCodeGen/Private/GenerateBlueprintCodeModuleCommandlet.cpp
--- a/Engine/Source/Developer/BlueprintNativeCodeGen/Private/GenerateBlueprintCodeModuleCommandlet.cpp
+++ b/Engine/Source/Developer/BlueprintNativeCodeGen/Private/GenerateBlueprintCodeModuleCommandlet.cpp
@@ -42,7 +42,7 @@ int32 UGenerateBlueprintCodeModuleCommandlet::Main(FString const& Params)
 	}
 
 	FBlueprintNativeCodeGenCoordinator Coordinator(CommandlineParams);
-	FBlueprintNativeCodeGenUtils::GenerateCodeModule(Coordinator);
+	//FBlueprintNativeCodeGenUtils::GenerateCodeModule(Coordinator);
 	
 	return 0;
 }
diff --git a/Engine/Source/Developer/BlueprintNativeCodeGen/Private/NativeCodeGenerationTool.cpp b/Engine/Source/Developer/BlueprintNativeCodeGen/Private/NativeCodeGenerationTool.cpp
--- a/Engine/Source/Developer/BlueprintNativeCodeGen/Private/NativeCodeGenerationTool.cpp
+++ b/Engine/Source/Developer/BlueprintNativeCodeGen/Private/NativeCodeGenerationTool.cpp
@@ -133,7 +133,7 @@ struct FGeneratedCodeData
 		{
 			TSharedPtr<FString> HeaderSource(new FString());
 			TSharedPtr<FString> CppSource(new FString());
-			FBlueprintNativeCodeGenUtils::GenerateCppCode(Obj, HeaderSource, CppSource);
+			//BFBlueprintNativeCodeGenUtils::GenerateCppCode(Obj, HeaderSource, CppSource);
 			SlowTask.EnterProgressFrame();
 
 			const FString FullHeaderFilename = FPaths::Combine(*HeaderDirPath, *(Obj->GetName() + TEXT(".h")));
