RootEditor_OnMac fixes. Likely required include changes due to 2684687

diff --git a/Engine/Source/Developer/BlueprintNativeCodeGen/Private/NativeCodeGenerationTool.cpp b/Engine/Source/Developer/BlueprintNativeCodeGen/Private/NativeCodeGenerationTool.cpp
--- a/Engine/Source/Developer/BlueprintNativeCodeGen/Private/NativeCodeGenerationTool.cpp
+++ b/Engine/Source/Developer/BlueprintNativeCodeGen/Private/NativeCodeGenerationTool.cpp
@@ -2,7 +2,9 @@
 
 #include "BlueprintNativeCodeGenPCH.h"
 #include "NativeCodeGenerationTool.h"
+#include "Editor.h"
 #include "Dialogs.h"
+#include "EditorStyle.h"
 #include "Engine/BlueprintGeneratedClass.h"
 #include "Engine/UserDefinedEnum.h"
 #include "Engine/UserDefinedStruct.h"
diff --git a/Engine/Source/Editor/Kismet/Public/BlueprintEditor.h b/Engine/Source/Editor/Kismet/Public/BlueprintEditor.h
--- a/Engine/Source/Editor/Kismet/Public/BlueprintEditor.h
+++ b/Engine/Source/Editor/Kismet/Public/BlueprintEditor.h
@@ -12,6 +12,7 @@
 #include "Engine/UserDefinedEnum.h"
 #include "PreviewScene.h"
 #include "Developer/Merge/Public/Merge.h" // for FOnMergeResolved
+#include "TickableEditorObject.h"
 
 class USCS_Node;
 
diff --git a/Engine/Source/Editor/UMGEditor/Private/AssetTypeActions_WidgetBlueprint.cpp b/Engine/Source/Editor/UMGEditor/Private/AssetTypeActions_WidgetBlueprint.cpp
--- a/Engine/Source/Editor/UMGEditor/Private/AssetTypeActions_WidgetBlueprint.cpp
+++ b/Engine/Source/Editor/UMGEditor/Private/AssetTypeActions_WidgetBlueprint.cpp
@@ -2,6 +2,7 @@
 
 #include "UMGEditorPrivatePCH.h"
 
+#include "Editor/UnrealEd/Classes/Settings/EditorExperimentalSettings.h"
 #include "PackageTools.h"
 #include "Kismet2/BlueprintEditorUtils.h"
 #include "Kismet2/KismetEditorUtilities.h"
