(WIP) Adding more commandline params for the native conversion tool (in setup for calling from the cook commandlet).

diff --git a/Engine/Source/Developer/BlueprintNativeCodeGen/Private/BlueprintNativeCodeGenManifest.cpp b/Engine/Source/Developer/BlueprintNativeCodeGen/Private/BlueprintNativeCodeGenManifest.cpp
--- a/Engine/Source/Developer/BlueprintNativeCodeGen/Private/BlueprintNativeCodeGenManifest.cpp
+++ b/Engine/Source/Developer/BlueprintNativeCodeGen/Private/BlueprintNativeCodeGenManifest.cpp
@@ -9,6 +9,7 @@
 #include "JsonObjectConverter.h"
 #include "Serialization/JsonWriter.h"
 
+DEFINE_LOG_CATEGORY_STATIC(LogNativeCodeGenManifest, Log, All);
 
 /*******************************************************************************
  * BlueprintNativeCodeGenManifestImpl
@@ -23,30 +24,21 @@ namespace BlueprintNativeCodeGenManifestImpl
 	static const FString HeaderSubDir    = TEXT("Public");
 	static const FString CppSubDir       = TEXT("Private");
 
-	/** */
-	static FString GetManifestFilePath(const FString& ModulePath);
-
 	/**  */
-	static bool LoadManifest(FBlueprintNativeCodeGenManifest* Manifest);
+	static bool LoadManifest(const FString& FilePath, FBlueprintNativeCodeGenManifest* Manifest);
 
 	/**  */
 	static FString GenerateHeaderSavePath(const FString& ModulePath, const FAssetData& Asset);
 	static FString GenerateCppSavePath(const FString& ModulePath, const FAssetData& Asset);
-}
 
-//------------------------------------------------------------------------------
-static FString BlueprintNativeCodeGenManifestImpl::GetManifestFilePath(const FString& ModulePath)
-{
-	return FPaths::Combine(*ModulePath, FApp::GetGameName()) + ManifestFileExt;
+	static FString GetComparibleDirPath(const FString& DirectoryPath);
 }
 
 //------------------------------------------------------------------------------
-static bool BlueprintNativeCodeGenManifestImpl::LoadManifest(FBlueprintNativeCodeGenManifest* Manifest)
+static bool BlueprintNativeCodeGenManifestImpl::LoadManifest(const FString& FilePath, FBlueprintNativeCodeGenManifest* Manifest)
 {
-	const FString ManifestFilePath = BlueprintNativeCodeGenManifestImpl::GetManifestFilePath(Manifest->GetTargetPath());
-
 	FString ManifestStr;
-	if (FFileHelper::LoadFileToString(ManifestStr, *ManifestFilePath))
+	if (FFileHelper::LoadFileToString(ManifestStr, *FilePath))
 	{
 		TSharedRef< TJsonReader<> > JsonReader = TJsonReaderFactory<>::Create(ManifestStr);
 
@@ -73,10 +65,39 @@ static FString BlueprintNativeCodeGenManifestImpl::GenerateCppSavePath(const FSt
 	return FPaths::Combine(*FPaths::Combine(*ModulePath, *CppSubDir), *Asset.AssetName.ToString()) + CppFileExt;
 }
 
+//------------------------------------------------------------------------------
+static FString BlueprintNativeCodeGenManifestImpl::GetComparibleDirPath(const FString& DirectoryPath)
+{
+	FString NormalizedPath = DirectoryPath;
+
+	const FString PathDelim = TEXT("/");
+	if (!NormalizedPath.EndsWith(PathDelim))
+	{
+		// to account for the case where the relative path would resolve to X: 
+		// (when we want "X:/")... ConvertRelativePathToFull() leaves the 
+		// trailing slash, and NormalizeDirectoryName() will remove it (if it is
+		// not a drive letter)
+		NormalizedPath += PathDelim;
+	}
+
+	if (FPaths::IsRelative(NormalizedPath))
+	{
+		NormalizedPath = FPaths::ConvertRelativePathToFull(NormalizedPath);
+	}
+	FPaths::NormalizeDirectoryName(NormalizedPath);
+	return NormalizedPath;
+}
+
 /*******************************************************************************
  * FBlueprintNativeCodeGenManifest
  ******************************************************************************/
 
+//------------------------------------------------------------------------------
+FString FBlueprintNativeCodeGenManifest::GetDefaultFilename()
+{
+	return FApp::GetGameName() + BlueprintNativeCodeGenManifestImpl::ManifestFileExt;
+}
+
 FBlueprintNativeCodeGenManifest::FBlueprintNativeCodeGenManifest(const FString TargetPath)
 	: ModulePath(TargetPath)
 {
@@ -91,14 +112,80 @@ FBlueprintNativeCodeGenManifest::FBlueprintNativeCodeGenManifest(const FString T
 //------------------------------------------------------------------------------
 FBlueprintNativeCodeGenManifest::FBlueprintNativeCodeGenManifest(const FNativeCodeGenCommandlineParams& CommandlineParams)
 {
+	using namespace BlueprintNativeCodeGenManifestImpl;
+	bool const bLoadExistingManifest = !CommandlineParams.bWipeRequested || CommandlineParams.ModuleOutputDir.IsEmpty();
+	
 	ModulePath = CommandlineParams.ModuleOutputDir;
+	if (FPaths::IsRelative(ModulePath))
+	{
+		FPaths::MakePathRelativeTo(ModulePath, *FPaths::GameDir());
+	}
+
+	if (!CommandlineParams.ManifestFilePath.IsEmpty())
+	{
+		ManifestPath = CommandlineParams.ManifestFilePath;
+	}
+	else
+	{
+		if (!ensure(!ModulePath.IsEmpty()))
+		{
+			ModulePath = FPaths::Combine(*FPaths::GameIntermediateDir(), TEXT("BpCodeGen"));
+		}
+		ManifestPath = FPaths::Combine(*ModulePath, *GetDefaultFilename());
+	}
+	
 	// incorporate an existing manifest (in case we're only re-converting a 
 	// handful of assets and adding them into an existing module)
-	if (BlueprintNativeCodeGenManifestImpl::LoadManifest(this))
+	if (bLoadExistingManifest && LoadManifest(ManifestPath, this))
+	{
+		// if they specified a separate module path, lets make sure we use that
+		// over what was found in the existing manifest
+		if (!CommandlineParams.ModuleOutputDir.IsEmpty())
+		{
+			if (CommandlineParams.bWipeRequested)
+			{
+				ModulePath = CommandlineParams.ModuleOutputDir;
+			}
+			else
+			{
+				const FString ExpectedPath = GetComparibleDirPath(CommandlineParams.ModuleOutputDir);
+				const FString TargetPath   = GetComparibleDirPath(GetTargetPath());
+
+				if ( !FPaths::IsSamePath(ExpectedPath, TargetPath) )
+				{
+					UE_LOG(LogNativeCodeGenManifest, Error, 
+						TEXT("The existing manifest's module path does not match what was specified via the commandline."));
+				}
+			}
+		}
+
+		// if we were only interested in obtaining the module path
+		if (CommandlineParams.bWipeRequested)
+		{
+			Clear();
+		}
+	}
+
+	if (CommandlineParams.bWipeRequested)
 	{
-		// reset ModulePath afterwards (in case what we've loaded has been moved
-		ModulePath = CommandlineParams.ModuleOutputDir;
+		const FString TargetPath = GetTargetPath();
+
+		IFileManager& FileManager = IFileManager::Get();
+		FileManager.DeleteDirectory(*FPaths::Combine(*TargetPath, *CppSubDir));
+		FileManager.DeleteDirectory(*FPaths::Combine(*TargetPath, *HeaderSubDir));
+	}
+}
+
+//------------------------------------------------------------------------------
+FString FBlueprintNativeCodeGenManifest::GetTargetPath() const
+{
+	FString TargetPath = ModulePath;
+	if (FPaths::IsRelative(ModulePath))
+	{
+		TargetPath = FPaths::ConvertRelativePathToFull(FPaths::GameDir(), ModulePath);
+		TargetPath = FPaths::ConvertRelativePathToFull(TargetPath);
 	}
+	return TargetPath;
 }
 
 //------------------------------------------------------------------------------
@@ -127,11 +214,16 @@ bool FBlueprintNativeCodeGenManifest::Save() const
 		if (FJsonSerializer::Serialize(JsonObject, JsonWriter))
 		{
 			JsonWriter->Close();
-
-			const FString ManifestFilePath = BlueprintNativeCodeGenManifestImpl::GetManifestFilePath(GetTargetPath());
-			return FFileHelper::SaveStringToFile(FileContents, *ManifestFilePath);
+			return FFileHelper::SaveStringToFile(FileContents, *ManifestPath);
 		}
 	}
 	return false;
 }
 
+//------------------------------------------------------------------------------
+void FBlueprintNativeCodeGenManifest::Clear()
+{
+	ConvertedAssets.Empty();
+	ModuleDependencies.Empty();
+}
+
diff --git a/Engine/Source/Developer/BlueprintNativeCodeGen/Private/BlueprintNativeCodeGenManifest.h b/Engine/Source/Developer/BlueprintNativeCodeGen/Private/BlueprintNativeCodeGenManifest.h
--- a/Engine/Source/Developer/BlueprintNativeCodeGen/Private/BlueprintNativeCodeGenManifest.h
+++ b/Engine/Source/Developer/BlueprintNativeCodeGen/Private/BlueprintNativeCodeGenManifest.h
@@ -40,11 +40,13 @@ struct FBlueprintNativeCodeGenManifest
 {
 	GENERATED_USTRUCT_BODY()
 
+	static FString GetDefaultFilename();
+
 public: 
 	FBlueprintNativeCodeGenManifest(const FString TargetPath = TEXT(""));
 	FBlueprintNativeCodeGenManifest(const FNativeCodeGenCommandlineParams& CommandlineParams);
 
-	FORCEINLINE const FString& GetTargetPath() const { return ModulePath; }
+	FString GetTargetPath() const;
 
 	/**  */
 	bool Save() const;
@@ -53,6 +55,12 @@ public:
 	FConvertedAssetRecord& CreateConversionRecord(const FAssetData& AssetInfo);
 
 private:
+	void Clear();
+
+	/**  */
+	FString ManifestPath;
+
+	/** Relative to the project's directory */
 	UPROPERTY()
 	FString ModulePath;
 
@@ -60,7 +68,5 @@ private:
 	TArray<FString> ModuleDependencies;
 
 	UPROPERTY()
-	TArray<FConvertedAssetRecord> ConvertedAssets; 
-	/**  */
-	TMap<FName, int32> RecordLookupTable;
+	TArray<FConvertedAssetRecord> ConvertedAssets;
 };
diff --git a/Engine/Source/Developer/BlueprintNativeCodeGen/Private/BlueprintNativeCodeGenUtils.cpp b/Engine/Source/Developer/BlueprintNativeCodeGen/Private/BlueprintNativeCodeGenUtils.cpp
--- a/Engine/Source/Developer/BlueprintNativeCodeGen/Private/BlueprintNativeCodeGenUtils.cpp
+++ b/Engine/Source/Developer/BlueprintNativeCodeGen/Private/BlueprintNativeCodeGenUtils.cpp
@@ -10,6 +10,7 @@
 #include "Engine/UserDefinedStruct.h"
 #include "Engine/UserDefinedEnum.h"
 #include "Kismet2/KismetEditorUtilities.h"		// for CompileBlueprint()
+#include "OutputDevice.h"						// for GWarn
 
 /*******************************************************************************
  * BlueprintNativeCodeGenUtilsImpl
@@ -20,6 +21,64 @@ namespace BlueprintNativeCodeGenUtilsImpl
 	
 }
 
+/*******************************************************************************
+ * FScopedFeedbackContext
+ ******************************************************************************/
+
+//------------------------------------------------------------------------------
+FBlueprintNativeCodeGenUtils::FScopedFeedbackContext::FScopedFeedbackContext()
+	: OldContext(GWarn)
+	, ErrorCount(0)
+	, WarningCount(0)
+{
+	TreatWarningsAsErrors = GWarn->TreatWarningsAsErrors;
+	GWarn = this;
+}
+
+//------------------------------------------------------------------------------
+FBlueprintNativeCodeGenUtils::FScopedFeedbackContext::~FScopedFeedbackContext()
+{
+	GWarn = OldContext;
+}
+
+//------------------------------------------------------------------------------
+bool FBlueprintNativeCodeGenUtils::FScopedFeedbackContext::HasErrors()
+{
+	return (ErrorCount > 0) || (TreatWarningsAsErrors && (WarningCount > 0));
+}
+
+//------------------------------------------------------------------------------
+void FBlueprintNativeCodeGenUtils::FScopedFeedbackContext::Serialize(const TCHAR* V, ELogVerbosity::Type Verbosity, const class FName& Category)
+{
+	switch (Verbosity)
+	{
+	case ELogVerbosity::Warning:
+		{
+			++WarningCount;
+		} 
+		break;
+
+	case ELogVerbosity::Error:
+	case ELogVerbosity::Fatal:
+		{
+			++ErrorCount;
+		} 
+		break;
+
+	default:
+		break;
+	}
+
+	OldContext->Serialize(V, Verbosity, Category);
+}
+
+//------------------------------------------------------------------------------
+void FBlueprintNativeCodeGenUtils::FScopedFeedbackContext::Flush()
+{
+	WarningCount = ErrorCount = 0;
+	OldContext->Flush();
+}
+
 /*******************************************************************************
  * FBlueprintNativeCodeGenUtils
  ******************************************************************************/
@@ -27,15 +86,15 @@ namespace BlueprintNativeCodeGenUtilsImpl
 //------------------------------------------------------------------------------
 bool FBlueprintNativeCodeGenUtils::GenerateCodeModule(FBlueprintNativeCodeGenCoordinator& Coordinator)
 {
-	
-	TSharedPtr<FString> HeaderSource(new FString());
-	TSharedPtr<FString> CppSource(new FString());
-
 	FBlueprintNativeCodeGenManifest& Manifest = Coordinator.Manifest;
 
+	TSharedPtr<FString> HeaderSource(new FString());
+	TSharedPtr<FString> CppSource(   new FString());
+
 	bool bSuccess = (Coordinator.ConversionQueue.Num() > 0);
 	for (const FAssetData& Asset : Coordinator.ConversionQueue)
 	{
+		FScopedFeedbackContext ScopedErrorTracker;
 		FConvertedAssetRecord& ConversionRecord = Manifest.CreateConversionRecord(Asset);
 
 		// loads the object if it is not already loaded
@@ -60,6 +119,12 @@ bool FBlueprintNativeCodeGenUtils::GenerateCodeModule(FBlueprintNativeCodeGenCoo
 		{
 			ConversionRecord.GeneratedCppPath.Empty();
 		}
+
+		if (ScopedErrorTracker.HasErrors())
+		{
+			bSuccess = false;
+			break;
+		}
 	}
 	Coordinator.Manifest.Save();
 
diff --git a/Engine/Source/Developer/BlueprintNativeCodeGen/Private/BlueprintNativeCodeGenUtils.h b/Engine/Source/Developer/BlueprintNativeCodeGen/Private/BlueprintNativeCodeGenUtils.h
--- a/Engine/Source/Developer/BlueprintNativeCodeGen/Private/BlueprintNativeCodeGenUtils.h
+++ b/Engine/Source/Developer/BlueprintNativeCodeGen/Private/BlueprintNativeCodeGenUtils.h
@@ -2,20 +2,72 @@
 
 #pragma once
 
+#include "FeedbackContext.h"
+
 // Forward declares
 struct FBlueprintNativeCodeGenCoordinator;
+class  SBuildProgressWidget;
 
 /**  */
 struct FBlueprintNativeCodeGenUtils
 {
-public: 
-	/**  */
+public:
+	/**
+	 * 
+	 * 
+	 * @param  Coordinator    
+	 * @return 
+	 */
 	static bool GenerateCodeModule(FBlueprintNativeCodeGenCoordinator& Coordinator);
 
-	/** 
-	 * Recompiles the bytecode of a blueprint only. Should only be run for 
-	 * recompiling dependencies during compile on load. 
+	/**
+	 * Recompiles the bytecode of a blueprint only. Should only be run for
+	 * recompiling dependencies during compile on load.
+	 * 
+	 * @param  Obj    
+	 * @param  OutHeaderSource    
+	 * @param  OutCppSource    
 	 */
 	static void GenerateCppCode(UObject* Obj, TSharedPtr<FString> OutHeaderSource, TSharedPtr<FString> OutCppSource);
+
+public: 
+	/** 
+	 * 
+	 */
+	class FScopedFeedbackContext : public FFeedbackContext
+	{
+	public:
+		FScopedFeedbackContext();
+		virtual ~FScopedFeedbackContext();
+
+		/**  */
+		bool HasErrors();
+
+		// FOutputDevice interface
+		virtual void Serialize(const TCHAR* V, ELogVerbosity::Type Verbosity, const class FName& Category) override;
+		virtual void Flush() override;
+		virtual void TearDown() override { OldContext->TearDown();  }
+		virtual bool CanBeUsedOnAnyThread() const override { return false; }
+		// End FOutputDevice interface
+
+		// FFeedbackContext interface
+		virtual bool YesNof(const FText& Question) override { return OldContext->YesNof(Question); }
+		virtual bool ReceivedUserCancel() override { return OldContext->ReceivedUserCancel(); }
+		virtual FContextSupplier* GetContext() const override { return OldContext->GetContext(); }
+		virtual void SetContext(FContextSupplier* InSupplier) override { OldContext->SetContext(InSupplier); }
+		virtual TWeakPtr<SBuildProgressWidget> ShowBuildProgressWindow() override { return OldContext->ShowBuildProgressWindow(); }
+		virtual void CloseBuildProgressWindow() override { OldContext->CloseBuildProgressWindow(); }
+
+	protected:
+		//virtual void StartSlowTask(const FText& Task, bool bShowCancelButton = false) override { OldContext->StartSlowTask(Task, bShowCancelButton); }
+		//virtual void FinalizeSlowTask() override { OldContext->FinalizeSlowTask(); }
+		//virtual void ProgressReported(const float TotalProgressInterp, FText DisplayMessage) override { OldContext->ProgressReported(TotalProgressInterp, DisplayMessage); }
+		// End FFeedbackContext interface
+
+	private:
+		FFeedbackContext* OldContext;
+		int32 ErrorCount;
+		int32 WarningCount;
+	};
 };
 
diff --git a/Engine/Source/Developer/BlueprintNativeCodeGen/Private/GenerateBlueprintCodeModuleCommandlet.cpp b/Engine/Source/Developer/BlueprintNativeCodeGen/Private/GenerateBlueprintCodeModuleCommandlet.cpp
--- a/Engine/Source/Developer/BlueprintNativeCodeGen/Private/GenerateBlueprintCodeModuleCommandlet.cpp
+++ b/Engine/Source/Developer/BlueprintNativeCodeGen/Private/GenerateBlueprintCodeModuleCommandlet.cpp
@@ -41,9 +41,13 @@ int32 UGenerateBlueprintCodeModuleCommandlet::Main(FString const& Params)
 		return 0;
 	}
 
+	FBlueprintNativeCodeGenUtils::FScopedFeedbackContext ScopedErrorTracker;
 	FBlueprintNativeCodeGenCoordinator Coordinator(CommandlineParams);
-	FBlueprintNativeCodeGenUtils::GenerateCodeModule(Coordinator);
-	
-	return 0;
+
+	if (!ScopedErrorTracker.HasErrors())
+	{
+		FBlueprintNativeCodeGenUtils::GenerateCodeModule(Coordinator);
+	}
+	return ScopedErrorTracker.HasErrors();
 }
 
diff --git a/Engine/Source/Developer/BlueprintNativeCodeGen/Private/NativeCodeGenCommandlineParams.cpp b/Engine/Source/Developer/BlueprintNativeCodeGen/Private/NativeCodeGenCommandlineParams.cpp
--- a/Engine/Source/Developer/BlueprintNativeCodeGen/Private/NativeCodeGenCommandlineParams.cpp
+++ b/Engine/Source/Developer/BlueprintNativeCodeGen/Private/NativeCodeGenCommandlineParams.cpp
@@ -3,6 +3,7 @@
 
 #include "BlueprintNativeCodeGenPCH.h"
 #include "NativeCodeGenCommandlineParams.h"
+#include "BlueprintNativeCodeGenManifest.h"
 
 DEFINE_LOG_CATEGORY_STATIC(LogNativeCodeGenCommandline, Log, All);
 
@@ -38,6 +39,20 @@ const FString FNativeCodeGenCommandlineParams::HelpMessage = TEXT("\n\
                         project's intermeadiate folder (within a sub-folder of \n\
                         its own). If specified as a relative path, it will be  \n\
                         relative to the target project's root directory.       \n\
+\n\
+    -wipe               Will clear out the target directory when specified.    \n\
+                        Normally, without this switch, if the target directory \n\
+                        contains source from a previous run, it builds atop the\n\
+                        existing module and appends to its existing manifest.  \n\
+\n\
+    -manifest=<Path>    Specifies where to save the conversion's manifest file.\n\
+                        If <Path> is not an existing directory or file, then it\n\
+                        is assumed that it is a file path. If a manifest       \n\
+                        already exists here, then we can use this in place of  \n\
+                        the -output command (and the target module directory   \n\
+                        will be determined by the old manifest file). If left  \n\
+                        unset, then it will default to the module's target     \n\
+                        directory.                                             \n\
 \n\
     -help, -h, -?       Display this message and then exit.                    \n\
 \n");
@@ -59,9 +74,10 @@ namespace NativeCodeGenCommandlineParamsImpl
 //------------------------------------------------------------------------------
 FNativeCodeGenCommandlineParams::FNativeCodeGenCommandlineParams(const TArray<FString>& CommandlineSwitches)
 	: bHelpRequested(false)
+	, bWipeRequested(false)
 {
 	IFileManager& FileManager = IFileManager::Get();
-	ModuleOutputDir = FPaths::Combine(*FPaths::GameIntermediateDir(), *NativeCodeGenCommandlineParamsImpl::DefaultModuleName);
+	const FString DefaultModulePath = FPaths::Combine(*FPaths::GameIntermediateDir(), *NativeCodeGenCommandlineParamsImpl::DefaultModuleName);
 
 	for (const FString& Param : CommandlineSwitches)
 	{
@@ -84,9 +100,10 @@ FNativeCodeGenCommandlineParams::FNativeCodeGenCommandlineParams(const TArray<FS
 		}
 		else if (!Switch.Compare(TEXT("output"), ESearchCase::IgnoreCase))
 		{
+			FPaths::NormalizeDirectoryName(Value);
 			// if it's a relative path, let's have it relative to the game 
 			// directory (not the UE executable)
-			if (FPaths::IsRelative(*Value))
+			if (FPaths::IsRelative(Value))
 			{
 				Value = FPaths::Combine(*FPaths::GameDir(), *Value);
 			}
@@ -97,7 +114,43 @@ FNativeCodeGenCommandlineParams::FNativeCodeGenCommandlineParams(const TArray<FS
 			}
 			else
 			{
-				UE_LOG(LogNativeCodeGenCommandline, Warning, TEXT("'%s' doesn't appear to be a usable output directory, defaulting to: '%s'"), *Value, *ModuleOutputDir);
+				UE_LOG(LogNativeCodeGenCommandline, Warning, TEXT("'%s' doesn't appear to be a valid output directory, defaulting to: '%s'"), *Value, *DefaultModulePath);
+			}
+		}
+		else if (!Switch.Compare(TEXT("wipe"), ESearchCase::IgnoreCase))
+		{
+			bWipeRequested = true;
+		}
+		else if (!Switch.Compare(TEXT("manifest"), ESearchCase::IgnoreCase))
+		{
+			FPaths::NormalizeDirectoryName(Value);
+
+			if (FileManager.FileExists(*Value))
+			{
+				ManifestFilePath = Value;
+			}
+			else if (FileManager.DirectoryExists(*Value))
+			{
+				ManifestFilePath = FPaths::Combine(*Value, *FBlueprintNativeCodeGenManifest::GetDefaultFilename());
+			}
+			else
+			{
+				UE_LOG(LogNativeCodeGenCommandline, Display, TEXT("Unsure if '%s' is supposed to denote a directory or a filename, assuming it's a file."), *Value);
+
+				FString FilePath = FPaths::GetPath(*Value);
+				if (FileManager.DirectoryExists(*FilePath) || FileManager.MakeDirectory(*FilePath))
+				{
+					ManifestFilePath = Value;
+				}
+				else
+				{
+					if (FPaths::IsRelative(Value))
+					{
+						Value = FPaths::ConvertRelativePathToFull(Value);
+					}
+					UE_LOG( LogNativeCodeGenCommandline, Warning, TEXT("'%s' doesn't appear to be a valid manifest file path, defaulting to: '%s'."), *Value,
+						*FPaths::Combine(*DefaultModulePath, *FBlueprintNativeCodeGenManifest::GetDefaultFilename()) );
+				}
 			}
 		}
 		else
@@ -105,5 +158,12 @@ FNativeCodeGenCommandlineParams::FNativeCodeGenCommandlineParams(const TArray<FS
 			//UE_LOG(LogNativeCodeGenCommandline, Warning, TEXT("Unrecognized commandline parameter: %s"), *Switch);
 		}
 	}
+
+	bool const bUtilizeExistingManifest = !ManifestFilePath.IsEmpty() && !bWipeRequested && FileManager.FileExists(*ManifestFilePath);
+	// an existing manifest would specify where to put the module
+	if (ModuleOutputDir.IsEmpty() && !bUtilizeExistingManifest)
+	{
+		ModuleOutputDir = DefaultModulePath;
+	}
 }
 
diff --git a/Engine/Source/Developer/BlueprintNativeCodeGen/Private/NativeCodeGenCommandlineParams.h b/Engine/Source/Developer/BlueprintNativeCodeGen/Private/NativeCodeGenCommandlineParams.h
--- a/Engine/Source/Developer/BlueprintNativeCodeGen/Private/NativeCodeGenCommandlineParams.h
+++ b/Engine/Source/Developer/BlueprintNativeCodeGen/Private/NativeCodeGenCommandlineParams.h
@@ -16,5 +16,7 @@ public:
 	TArray<FString> WhiteListedAssetPaths;
 	TArray<FString> BlackListedAssetPaths;
 	FString         ModuleOutputDir;
+	FString			ManifestFilePath;
+	bool            bWipeRequested : 1;
 };
 
