(WIP) Workaround to synchronize #include names and saved filenames for converted Blueprint files (both should go through a centralized function).

diff --git a/Engine/Source/Developer/BlueprintNativeCodeGen/Private/BlueprintNativeCodeGenManifest.cpp b/Engine/Source/Developer/BlueprintNativeCodeGen/Private/BlueprintNativeCodeGenManifest.cpp
--- a/Engine/Source/Developer/BlueprintNativeCodeGen/Private/BlueprintNativeCodeGenManifest.cpp
+++ b/Engine/Source/Developer/BlueprintNativeCodeGen/Private/BlueprintNativeCodeGenManifest.cpp
@@ -27,6 +27,8 @@ namespace BlueprintNativeCodeGenManifestImpl
 	/**  */
 	static bool LoadManifest(const FString& FilePath, FBlueprintNativeCodeGenManifest* Manifest);
 
+	static FString GetBaseFilename(const FAssetData& Asset);
+
 	/**  */
 	static FString GenerateHeaderSavePath(const FString& ModulePath, const FAssetData& Asset);
 	static FString GenerateCppSavePath(const FString& ModulePath, const FAssetData& Asset);
@@ -52,17 +54,29 @@ static bool BlueprintNativeCodeGenManifestImpl::LoadManifest(const FString& File
 	return false;
 }
 
+//------------------------------------------------------------------------------
+static FString BlueprintNativeCodeGenManifestImpl::GetBaseFilename(const FAssetData& Asset)
+{
+	// @TODO: Coordinate this with the CppBackend module (so it matches includes)
+	FString Filename = Asset.AssetName.ToString();
+	if ( Asset.GetClass()->IsChildOf(UBlueprint::StaticClass()) )
+	{
+		static const FString BlueprintClassPostfix(TEXT("_C"));
+		Filename += BlueprintClassPostfix;
+	}
+	return Filename;
+}
+
 //------------------------------------------------------------------------------
 static FString BlueprintNativeCodeGenManifestImpl::GenerateHeaderSavePath(const FString& ModulePath, const FAssetData& Asset)
 {
-	// @TODO: Coordinate this with the CppBackend module
-	return FPaths::Combine(*FPaths::Combine(*ModulePath, *HeaderSubDir), *Asset.AssetName.ToString()) + HeaderFileExt;
+	return FPaths::Combine(*FPaths::Combine(*ModulePath, *HeaderSubDir), *GetBaseFilename(Asset)) + HeaderFileExt;
 }
 
 //------------------------------------------------------------------------------
 static FString BlueprintNativeCodeGenManifestImpl::GenerateCppSavePath(const FString& ModulePath, const FAssetData& Asset)
 {
-	return FPaths::Combine(*FPaths::Combine(*ModulePath, *CppSubDir), *Asset.AssetName.ToString()) + CppFileExt;
+	return FPaths::Combine(*FPaths::Combine(*ModulePath, *CppSubDir), *GetBaseFilename(Asset)) + CppFileExt;
 }
 
 //------------------------------------------------------------------------------
@@ -194,8 +208,8 @@ FConvertedAssetRecord& FBlueprintNativeCodeGenManifest::CreateConversionRecord(c
 	FConvertedAssetRecord NewConversionRecord;
 	NewConversionRecord.AssetType = AssetInfo.GetClass();
 	NewConversionRecord.AssetPath = AssetInfo.PackagePath.ToString();
-	NewConversionRecord.GeneratedHeaderPath = BlueprintNativeCodeGenManifestImpl::GenerateHeaderSavePath(ModulePath, AssetInfo);
-	NewConversionRecord.GeneratedCppPath = BlueprintNativeCodeGenManifestImpl::GenerateCppSavePath(ModulePath, AssetInfo);
+	NewConversionRecord.GeneratedHeaderPath = BlueprintNativeCodeGenManifestImpl::GenerateHeaderSavePath(GetTargetPath(), AssetInfo);
+	NewConversionRecord.GeneratedCppPath    = BlueprintNativeCodeGenManifestImpl::GenerateCppSavePath(GetTargetPath(), AssetInfo);
 
 	return ConvertedAssets[ ConvertedAssets.Add(NewConversionRecord) ];
 }
