Allowing code generation for blueprints with warnings

diff --git a/Engine/Source/Developer/BlueprintNativeCodeGen/Private/NativeCodeGenerationTool.cpp b/Engine/Source/Developer/BlueprintNativeCodeGen/Private/NativeCodeGenerationTool.cpp
--- a/Engine/Source/Developer/BlueprintNativeCodeGen/Private/NativeCodeGenerationTool.cpp
+++ b/Engine/Source/Developer/BlueprintNativeCodeGen/Private/NativeCodeGenerationTool.cpp
@@ -442,7 +442,7 @@ void FNativeCodeGenerationTool::Open(UBlueprint& Blueprint, TSharedRef< class FB
 
 bool FNativeCodeGenerationTool::CanGenerate(const UBlueprint& Blueprint)
 {
-	return (Blueprint.Status == EBlueprintStatus::BS_UpToDate)
+	return (Blueprint.Status == EBlueprintStatus::BS_UpToDate || Blueprint.Status == EBlueprintStatus::BS_UpToDateWithWarnings)
 		&& (Blueprint.BlueprintType == EBlueprintType::BPTYPE_Normal || Blueprint.BlueprintType == EBlueprintType::BPTYPE_FunctionLibrary)
 		&& Cast<UBlueprintGeneratedClass>(Blueprint.GeneratedClass);
 }
