(WIP) Blueprint conversion commandlet: gathering all required dependencies,  hooking into the backend (communicating which assets will be converted, etc.).

#codereviw Maciej.Mroz

diff --git a/Engine/Source/Developer/BlueprintNativeCodeGen/BlueprintNativeCodeGen.Build.cs b/Engine/Source/Developer/BlueprintNativeCodeGen/BlueprintNativeCodeGen.Build.cs
--- a/Engine/Source/Developer/BlueprintNativeCodeGen/BlueprintNativeCodeGen.Build.cs
+++ b/Engine/Source/Developer/BlueprintNativeCodeGen/BlueprintNativeCodeGen.Build.cs
@@ -28,6 +28,7 @@ public class BlueprintNativeCodeGen : ModuleRules
                 "Json",
                 "JsonUtilities",
                 "BlueprintCompilerCppBackend",
+                "GameProjectGeneration"
 			}
 		);
 
diff --git a/Engine/Source/Developer/BlueprintNativeCodeGen/Private/BlueprintNativeCodeGenCoordinator.cpp b/Engine/Source/Developer/BlueprintNativeCodeGen/Private/BlueprintNativeCodeGenCoordinator.cpp
--- a/Engine/Source/Developer/BlueprintNativeCodeGen/Private/BlueprintNativeCodeGenCoordinator.cpp
+++ b/Engine/Source/Developer/BlueprintNativeCodeGen/Private/BlueprintNativeCodeGenCoordinator.cpp
@@ -10,6 +10,9 @@
 #include "PackageName.h"
 #include "FileManager.h"
 #include "AssetRegistryModule.h"
+#include "BlueprintNativeCodeGenUtils.h"
+#include "BlueprintCompilerCppBackendGatherDependencies.h"
+#include "BlueprintSupport.h"	// for FReplaceCookedBPGC::OnQueryIfAssetIsTargetedForConversion()
 
 /*******************************************************************************
  * UBlueprintNativeCodeGenConfig
@@ -46,11 +49,17 @@ namespace BlueprintNativeCodeGenCoordinatorImpl
 	/**  */
 	static void GatherAssetsToConvert(const FNativeCodeGenCommandlineParams& CommandlineParams, TArray<FAssetData>& ConversionQueueOut);
 
-	/** */
+	/**  */
 	static void GatherParentAssets(const FAssetData& TargetAsset, TArray<FAssetData>& AssetsOut);
 
-	/** */
+	/**  */
 	static void GatherInterfaceAssets(const FAssetData& TargetAsset, TArray<FAssetData>& AssetsOut);
+
+	/**  */
+	static UStruct* GetAssetStruct(UObject* AssetObj);
+
+	/**  */
+	static void GatherInnerDependencies(UObject* AssetObj, TArray<FAssetData>& AssetsOut);
 }
 
 /*******************************************************************************
@@ -183,7 +192,16 @@ void FSpecializedAssetCollector<UBlueprint>::GatherAssets(const FARFilter& Share
 	{
 		BlueprintNativeCodeGenCoordinatorImpl::GatherParentAssets(AssetsOut[BlueprintIndex], AssetsOut);
 		BlueprintNativeCodeGenCoordinatorImpl::GatherInterfaceAssets(AssetsOut[BlueprintIndex], AssetsOut);
-		// @TODO: Collect UserDefinedStructs (maybe enums?)... maybe put off till later, once the BP is loaded?
+
+		// GatherInnerDependencies() will gather other required depenencies 
+		// (like user-defined structs/enums), but we don't want to execute that
+		// quite yet; 
+		//     1) because we'd have to load the asset for that here, and...
+		//     2) we want to do it later, when we're processing the conversion 
+		//        queue, once the CppBackend's 
+		//        OnCheckIsAssetTargetedForConversion() is setup to prevent 
+		//        unwanted Blueprint classes from being converted
+		// 
 	}
 }
 
@@ -230,7 +248,7 @@ static void BlueprintNativeCodeGenCoordinatorImpl::GatherAssetsToConvert(const F
 	SetFilterPackages(ConfigSettings->PackagesToAlwaysConvert);
 	SetFilterPackages(CommandlineParams.WhiteListedAssetPaths);
 
-	// will be utilized if UBlueprint is an entry in TargetAssetTypes (auto registers
+	// will be utilized if UBlueprint is an entry in TargetAssetTypes (auto registers itself as a handler)
 	FSpecializedAssetCollector<UBlueprint> BlueprintAssetCollector;
 
 	const int32 TargetClassCount = sizeof(TargetAssetTypes) / sizeof(TargetAssetTypes[0]);
@@ -316,6 +334,59 @@ static void BlueprintNativeCodeGenCoordinatorImpl::GatherInterfaceAssets(const F
 	}
 }
 
+//------------------------------------------------------------------------------
+static UStruct* BlueprintNativeCodeGenCoordinatorImpl::GetAssetStruct(UObject* AssetObj)
+{
+	if (UStruct* StructAsset = Cast<UStruct>(AssetObj))
+	{
+		return StructAsset;
+	}
+	else if (UBlueprint* BlueprintAsset = Cast<UBlueprint>(AssetObj))
+	{
+		return BlueprintAsset->GeneratedClass;
+	}
+	return nullptr;
+}
+
+//------------------------------------------------------------------------------
+static void BlueprintNativeCodeGenCoordinatorImpl::GatherInnerDependencies(UObject* AssetObj, TArray<FAssetData>& AssetsOut)
+{
+	if (UStruct* StructAsset = GetAssetStruct(AssetObj))
+	{
+		IAssetRegistry& AssetRegistry = GetAssetRegistry();
+		auto AddDependencyForConversion = [&AssetsOut, &AssetRegistry](UObject* AssetObj)->bool
+		{
+			ensure(AssetObj->IsAsset());
+			FAssetData AssetInfo = AssetRegistry.GetAssetByObjectPath(*AssetObj->GetPathName());
+
+			const bool bAssetFound = AssetInfo.IsValid();
+			if (bAssetFound)
+			{
+				AssetsOut.AddUnique(AssetInfo);
+			}
+			return bAssetFound;
+		};
+
+		FGatherConvertedClassDependencies AssetDependencies(StructAsset);
+
+		for (UBlueprintGeneratedClass* BlueprintClass : AssetDependencies.ConvertedClasses)
+		{
+			if (UBlueprint* Blueprint = Cast<UBlueprint>(BlueprintClass->ClassGeneratedBy))
+			{
+				AddDependencyForConversion(Blueprint);
+			}
+		}
+		for (UUserDefinedStruct* StructAsset : AssetDependencies.ConvertedStructs)
+		{
+			AddDependencyForConversion(StructAsset);
+		}
+		for (UUserDefinedEnum* EnumAsset : AssetDependencies.ConvertedEnum)
+		{
+			AddDependencyForConversion(EnumAsset);
+		}
+	}
+}
+
 /*******************************************************************************
  * FBlueprintNativeCodeGenCoordinator
 *******************************************************************************/
@@ -330,7 +401,36 @@ FBlueprintNativeCodeGenCoordinator::FBlueprintNativeCodeGenCoordinator(const FNa
 //------------------------------------------------------------------------------
 bool FBlueprintNativeCodeGenCoordinator::IsTargetedForConversion(const FString AssetPath)
 {
-	return ensure(false); // @TODO
+	bool bHasBeenConverted = false;
+	if (FConvertedAssetRecord* ConversionRecord = Manifest.FindConversionRecord(AssetPath))
+	{
+		// interfaces don't generate cpp files
+		bHasBeenConverted = !ConversionRecord->GeneratedHeaderPath.IsEmpty() || !ConversionRecord->GeneratedCppPath.IsEmpty();
+		
+		IFileManager& FileManager = IFileManager::Get();
+		bHasBeenConverted &= ConversionRecord->GeneratedHeaderPath.IsEmpty() || FileManager.FileExists(*ConversionRecord->GeneratedHeaderPath);
+		bHasBeenConverted &= ConversionRecord->GeneratedCppPath.IsEmpty()    || FileManager.FileExists(*ConversionRecord->GeneratedCppPath);
+	}
+
+	bool bIsQueuedForConversion = false;
+	if (!bHasBeenConverted)
+	{
+		// @TODO: consider using a map to speed this up
+
+		IAssetRegistry& AssetRegistry = BlueprintNativeCodeGenCoordinatorImpl::GetAssetRegistry();
+		TArray<FAssetData> AssetPackages;
+		AssetRegistry.GetAssetsByPackageName(*AssetPath, AssetPackages);
+
+		// should either be 0 (the asset doesn't exist), or 1 (the package contains a singular asset) 
+		ensure(AssetPackages.Num() <= 1);
+
+		for (const FAssetData& FoundAsset : AssetPackages)
+		{
+			// we shouldn't come up with multiple assets, but if we do... I guess OR the result?
+			bIsQueuedForConversion |= (ConversionQueue.Find(FoundAsset) != INDEX_NONE);
+		}
+	}
+	return bHasBeenConverted || bIsQueuedForConversion;
 }
 
 //------------------------------------------------------------------------------
@@ -373,3 +473,36 @@ bool FBlueprintNativeCodeGenCoordinator::IsTargetedForConversion(const UStruct*
 	return false;
 }
 
+//------------------------------------------------------------------------------
+bool FBlueprintNativeCodeGenCoordinator::ProcessConversionQueue(const FConversionDelegate& ConversionDelegate)
+{
+	auto& ConversionQueryDelegate = FReplaceCookedBPGC::Get().OnQueryIfAssetIsTargetedForConversion();
+	auto IsAssetTargetedForConversion = [this](const UObject* AssetObj)->bool
+	{
+		return IsTargetedForConversion(AssetObj->GetOutermost()->GetPathName());
+	};
+	ConversionQueryDelegate.BindLambda(IsAssetTargetedForConversion);
+
+	bool bSuccess = true;
+	for (int32 AssetIndex = 0; AssetIndex < ConversionQueue.Num(); ++AssetIndex)
+	{
+		const FAssetData& AssetInfo = ConversionQueue[AssetIndex];
+		// load asset, add dependencies that can only be gathered from a 
+		// loaded asset (property types, etc.)
+		UObject* AssetObj = AssetInfo.GetAsset();
+		BlueprintNativeCodeGenCoordinatorImpl::GatherInnerDependencies(AssetObj, ConversionQueue);
+
+		FConvertedAssetRecord& ConversionRecord = Manifest.CreateConversionRecord(AssetInfo);
+		bSuccess &= ConversionDelegate.Execute(ConversionRecord);
+
+		if (!bSuccess)
+		{
+			UE_LOG(LogBlueprintCodeGen, Error, TEXT("Failed to fully convert '%s'"), *AssetInfo.AssetName.ToString());
+			break;
+		}
+	}
+
+	ConversionQueryDelegate.Unbind();
+	return bSuccess;
+}
+
diff --git a/Engine/Source/Developer/BlueprintNativeCodeGen/Private/BlueprintNativeCodeGenCoordinator.h b/Engine/Source/Developer/BlueprintNativeCodeGen/Private/BlueprintNativeCodeGenCoordinator.h
--- a/Engine/Source/Developer/BlueprintNativeCodeGen/Private/BlueprintNativeCodeGenCoordinator.h
+++ b/Engine/Source/Developer/BlueprintNativeCodeGen/Private/BlueprintNativeCodeGenCoordinator.h
@@ -8,7 +8,9 @@
 
 // Forward declarations
 struct FNativeCodeGenCommandlineParams;
+struct FConvertedAssetRecord;
 
+/**  */
 UCLASS(config=Editor)
 class UBlueprintNativeCodeGenConfig : public UObject
 {
@@ -25,14 +27,21 @@ struct FBlueprintNativeCodeGenCoordinator
 public:
 	FBlueprintNativeCodeGenCoordinator(const FNativeCodeGenCommandlineParams& CommandlineParams);
 
+	const FBlueprintNativeCodeGenManifest& GetManifest() { return Manifest; }
+
 	bool IsTargetedForConversion(const FString AssetPath);
 	bool IsTargetedForConversion(const UBlueprint* Blueprint);
 	bool IsTargetedForConversion(const UClass* Class);
 	bool IsTargetedForConversion(const UEnum* Enum);
 	bool IsTargetedForConversion(const UStruct* Struct);
 
-public:
+	DECLARE_DELEGATE_RetVal_OneParam(bool, FConversionDelegate, FConvertedAssetRecord&);
+	bool ProcessConversionQueue(const FConversionDelegate& ConversionDelegate);
+
+private:
+	/**  */
 	TArray<FAssetData> ConversionQueue;
+	/**  */
 	FBlueprintNativeCodeGenManifest Manifest;
 };
 
diff --git a/Engine/Source/Developer/BlueprintNativeCodeGen/Private/BlueprintNativeCodeGenManifest.cpp b/Engine/Source/Developer/BlueprintNativeCodeGen/Private/BlueprintNativeCodeGenManifest.cpp
--- a/Engine/Source/Developer/BlueprintNativeCodeGen/Private/BlueprintNativeCodeGenManifest.cpp
+++ b/Engine/Source/Developer/BlueprintNativeCodeGen/Private/BlueprintNativeCodeGenManifest.cpp
@@ -8,6 +8,7 @@
 #include "Serialization/JsonSerializer.h"
 #include "JsonObjectConverter.h"
 #include "Serialization/JsonWriter.h"
+#include "IBlueprintCompilerCppBackendModule.h" // for GetBaseFilename()
 
 DEFINE_LOG_CATEGORY_STATIC(LogNativeCodeGenManifest, Log, All);
 
@@ -27,6 +28,8 @@ namespace BlueprintNativeCodeGenManifestImpl
 	/**  */
 	static bool LoadManifest(const FString& FilePath, FBlueprintNativeCodeGenManifest* Manifest);
 
+	static FString GetHeaderSaveDir(const FString& ModulePath);
+	static FString GetCppSaveDir(const FString& ModulePath);
 	static FString GetBaseFilename(const FAssetData& Asset);
 
 	/**  */
@@ -54,29 +57,34 @@ static bool BlueprintNativeCodeGenManifestImpl::LoadManifest(const FString& File
 	return false;
 }
 
+//------------------------------------------------------------------------------
+static FString BlueprintNativeCodeGenManifestImpl::GetHeaderSaveDir(const FString& ModulePath)
+{
+	return FPaths::Combine(*ModulePath, *HeaderSubDir);
+}
+
+//------------------------------------------------------------------------------
+static FString BlueprintNativeCodeGenManifestImpl::GetCppSaveDir(const FString& ModulePath)
+{
+	return FPaths::Combine(*ModulePath, *CppSubDir);
+}
+
 //------------------------------------------------------------------------------
 static FString BlueprintNativeCodeGenManifestImpl::GetBaseFilename(const FAssetData& Asset)
 {
-	// @TODO: Coordinate this with the CppBackend module (so it matches includes)
-	FString Filename = Asset.AssetName.ToString();
-	if ( Asset.GetClass()->IsChildOf(UBlueprint::StaticClass()) )
-	{
-		static const FString BlueprintClassPostfix(TEXT("_C"));
-		Filename += BlueprintClassPostfix;
-	}
-	return Filename;
+	return IBlueprintCompilerCppBackendModule::GetBaseFilename(Asset.GetAsset());
 }
 
 //------------------------------------------------------------------------------
 static FString BlueprintNativeCodeGenManifestImpl::GenerateHeaderSavePath(const FString& ModulePath, const FAssetData& Asset)
 {
-	return FPaths::Combine(*FPaths::Combine(*ModulePath, *HeaderSubDir), *GetBaseFilename(Asset)) + HeaderFileExt;
+	return FPaths::Combine(*GetHeaderSaveDir(ModulePath), *GetBaseFilename(Asset)) + HeaderFileExt;
 }
 
 //------------------------------------------------------------------------------
 static FString BlueprintNativeCodeGenManifestImpl::GenerateCppSavePath(const FString& ModulePath, const FAssetData& Asset)
 {
-	return FPaths::Combine(*FPaths::Combine(*ModulePath, *CppSubDir), *GetBaseFilename(Asset)) + CppFileExt;
+	return FPaths::Combine(*GetCppSaveDir(ModulePath), *GetBaseFilename(Asset)) + CppFileExt;
 }
 
 //------------------------------------------------------------------------------
@@ -145,7 +153,7 @@ FBlueprintNativeCodeGenManifest::FBlueprintNativeCodeGenManifest(const FNativeCo
 		{
 			ModulePath = FPaths::Combine(*FPaths::GameIntermediateDir(), TEXT("BpCodeGen"));
 		}
-		ManifestPath = FPaths::Combine(*ModulePath, *GetDefaultFilename());
+		ManifestPath = FPaths::Combine(*GetModuleDir(), *GetDefaultFilename());
 	}
 	
 	// incorporate an existing manifest (in case we're only re-converting a 
@@ -163,7 +171,9 @@ FBlueprintNativeCodeGenManifest::FBlueprintNativeCodeGenManifest(const FNativeCo
 			else
 			{
 				const FString ExpectedPath = GetComparibleDirPath(CommandlineParams.ModuleOutputDir);
-				const FString TargetPath   = GetComparibleDirPath(GetTargetPath());
+
+				FString TargetPath = GetModuleDir();
+				TargetPath = GetComparibleDirPath(TargetPath);
 
 				if ( !FPaths::IsSamePath(ExpectedPath, TargetPath) )
 				{
@@ -178,45 +188,87 @@ FBlueprintNativeCodeGenManifest::FBlueprintNativeCodeGenManifest(const FNativeCo
 		{
 			Clear();
 		}
+		else
+		{
+			MapConvertedAssets();
+		}
 	}
+}
 
-	if (CommandlineParams.bWipeRequested)
+//------------------------------------------------------------------------------
+FConvertedAssetRecord& FBlueprintNativeCodeGenManifest::CreateConversionRecord(const FAssetData& AssetInfo)
+{
+	const FString AssetPath = AssetInfo.PackageName.ToString();
+	if (FConvertedAssetRecord* ExistingConversionRecord = FindConversionRecord(AssetPath))
 	{
-		const FString TargetPath = GetTargetPath();
-
-		IFileManager& FileManager = IFileManager::Get();
-		FileManager.DeleteDirectory(*FPaths::Combine(*TargetPath, *CppSubDir));
-		FileManager.DeleteDirectory(*FPaths::Combine(*TargetPath, *HeaderSubDir));
+		return *ExistingConversionRecord;
 	}
+
+	const FString TargetModulePath = GetModuleDir();
+
+	FConvertedAssetRecord NewConversionRecord;
+	NewConversionRecord.AssetPtr  = AssetInfo.GetAsset();
+	NewConversionRecord.AssetType = AssetInfo.GetClass();
+	NewConversionRecord.AssetPath = AssetPath;
+	NewConversionRecord.GeneratedHeaderPath = BlueprintNativeCodeGenManifestImpl::GenerateHeaderSavePath(TargetModulePath, AssetInfo);
+	NewConversionRecord.GeneratedCppPath    = BlueprintNativeCodeGenManifestImpl::GenerateCppSavePath(TargetModulePath, AssetInfo);
+
+	RecordMap.Add(AssetPath, ConvertedAssets.Num());
+	return ConvertedAssets[ ConvertedAssets.Add(NewConversionRecord) ];
 }
 
 //------------------------------------------------------------------------------
-FString FBlueprintNativeCodeGenManifest::GetTargetPath() const
+FConvertedAssetRecord* FBlueprintNativeCodeGenManifest::FindConversionRecord(const FString& AssetPath, bool bSlowSearch)
 {
-	FString TargetPath = ModulePath;
-	if (FPaths::IsRelative(ModulePath))
+	FConvertedAssetRecord* FoundRecord = nullptr;
+	if (int32* IndexPtr = RecordMap.Find(AssetPath))
 	{
-		TargetPath = FPaths::ConvertRelativePathToFull(FPaths::GameDir(), ModulePath);
-		TargetPath = FPaths::ConvertRelativePathToFull(TargetPath);
+		check(*IndexPtr < ConvertedAssets.Num());
+		FoundRecord = &ConvertedAssets[*IndexPtr];
+		ensure(FoundRecord->AssetPath == AssetPath);
 	}
-	return TargetPath;
+	else if (bSlowSearch)
+	{
+		for (FConvertedAssetRecord& ConversionRecord : ConvertedAssets)
+		{
+			if (ConversionRecord.AssetPath == AssetPath)
+			{
+				FoundRecord = &ConversionRecord;
+			}
+		}
+	}
+	return FoundRecord;
 }
 
 //------------------------------------------------------------------------------
-FConvertedAssetRecord& FBlueprintNativeCodeGenManifest::CreateConversionRecord(const FAssetData& AssetInfo)
+TArray<FString> FBlueprintNativeCodeGenManifest::GetTargetPaths() const 
 {
-	FConvertedAssetRecord NewConversionRecord;
-	NewConversionRecord.AssetType = AssetInfo.GetClass();
-	NewConversionRecord.AssetPath = AssetInfo.PackagePath.ToString();
-	NewConversionRecord.GeneratedHeaderPath = BlueprintNativeCodeGenManifestImpl::GenerateHeaderSavePath(GetTargetPath(), AssetInfo);
-	NewConversionRecord.GeneratedCppPath    = BlueprintNativeCodeGenManifestImpl::GenerateCppSavePath(GetTargetPath(), AssetInfo);
+	const FString TargetModulePath = GetModuleDir();
 
-	return ConvertedAssets[ ConvertedAssets.Add(NewConversionRecord) ];
+	TArray<FString> DestPaths;
+	DestPaths.Add(ManifestPath);
+	DestPaths.Add(BlueprintNativeCodeGenManifestImpl::GetHeaderSaveDir(TargetModulePath));
+	DestPaths.Add(BlueprintNativeCodeGenManifestImpl::GetCppSaveDir(TargetModulePath));
+	// @TODO: Add *Build.cs file
+
+	return DestPaths;
 }
 
 //------------------------------------------------------------------------------
-bool FBlueprintNativeCodeGenManifest::Save() const
+bool FBlueprintNativeCodeGenManifest::Save(bool bSort) const
 {
+	if (bSort)
+	{
+		ConvertedAssets.Sort([](const FConvertedAssetRecord& Lhs, const FConvertedAssetRecord& Rhs)->bool
+		{ 
+			if (Lhs.AssetType == Rhs.AssetType)
+			{
+				return Lhs.AssetPath < Rhs.AssetPath;
+			}
+			return Lhs.AssetType->GetName() < Rhs.AssetType->GetName();
+		});
+	}
+
 	TSharedRef<FJsonObject> JsonObject = MakeShareable(new FJsonObject());
 
 	if (FJsonObjectConverter::UStructToJsonObject(FBlueprintNativeCodeGenManifest::StaticStruct(), this, JsonObject,
@@ -238,6 +290,39 @@ bool FBlueprintNativeCodeGenManifest::Save() const
 void FBlueprintNativeCodeGenManifest::Clear()
 {
 	ConvertedAssets.Empty();
-	ModuleDependencies.Empty();
+	RecordMap.Empty();
+}
+
+//------------------------------------------------------------------------------
+FString FBlueprintNativeCodeGenManifest::GetModuleDir() const
+{
+	FString TargetPath = ModulePath;
+	if (FPaths::IsRelative(ModulePath))
+	{
+		TargetPath = FPaths::ConvertRelativePathToFull(FPaths::GameDir(), ModulePath);
+		TargetPath = FPaths::ConvertRelativePathToFull(TargetPath);
+	}
+	return TargetPath;
+}
+
+//------------------------------------------------------------------------------
+void FBlueprintNativeCodeGenManifest::MapConvertedAssets()
+{
+	RecordMap.Empty(ConvertedAssets.Num());
+
+	for (int32 RecordIndex = 0; RecordIndex < ConvertedAssets.Num(); ++RecordIndex)
+	{
+		const FConvertedAssetRecord& AssetRecord = ConvertedAssets[RecordIndex];
+		RecordMap.Add(AssetRecord.AssetPath, RecordIndex);
+		
+// 		if (UObject* AssetObj = AssetRecord.AssetPtr.Get())
+// 		{
+//			RecordMap.Add(AssetRecord.AssetPath, RecordIndex);
+// 		}
+// 		else
+// 		{
+// 			UE_LOG(LogBlueprintCodeGen, Warning, TEXT("Invalid conversion record found for: %s"), *AssetRecord.AssetPath);
+// 		}
+	}
 }
 
diff --git a/Engine/Source/Developer/BlueprintNativeCodeGen/Private/BlueprintNativeCodeGenManifest.h b/Engine/Source/Developer/BlueprintNativeCodeGen/Private/BlueprintNativeCodeGenManifest.h
--- a/Engine/Source/Developer/BlueprintNativeCodeGen/Private/BlueprintNativeCodeGenManifest.h
+++ b/Engine/Source/Developer/BlueprintNativeCodeGen/Private/BlueprintNativeCodeGenManifest.h
@@ -18,6 +18,8 @@ struct FConvertedAssetRecord
 	GENERATED_USTRUCT_BODY()
 
 public:
+	TWeakObjectPtr<UObject> AssetPtr;
+
 	UPROPERTY()
 	UClass* AssetType;
 
@@ -46,17 +48,50 @@ public:
 	FBlueprintNativeCodeGenManifest(const FString TargetPath = TEXT(""));
 	FBlueprintNativeCodeGenManifest(const FNativeCodeGenCommandlineParams& CommandlineParams);
 
-	FString GetTargetPath() const;
+	/**
+	 * 
+	 * 
+	 * @return 
+	 */
+	TArray<FString> GetTargetPaths() const;
+
+	/**
+	 * 
+	 * 
+	 * @param  AssetInfo    
+	 * @return 
+	 */
+	FConvertedAssetRecord& CreateConversionRecord(const FAssetData& AssetInfo);
 
-	/**  */
-	bool Save() const;
+	/**
+	 * 
+	 * 
+	 * @param  AssetPath    
+	 * @return 
+	 */
+	FConvertedAssetRecord* FindConversionRecord(const FString& AssetPath, bool bSlowSearch = false);
 
-	/**  */
-	FConvertedAssetRecord& CreateConversionRecord(const FAssetData& AssetInfo);
+	/**
+	 * 
+	 */
+	bool Save(bool bSort = true) const;
 
 private:
+	/**
+	 * 
+	 */
+	FString GetModuleDir() const;
+
+	/**
+	 * 
+	 */
 	void Clear();
 
+	/**
+	 *
+	 */
+	void MapConvertedAssets();
+	
 	/**  */
 	FString ManifestPath;
 
@@ -64,9 +99,10 @@ private:
 	UPROPERTY()
 	FString ModulePath;
 
+	/** Mutable so Save() can sort the array when requested */
 	UPROPERTY()
-	TArray<FString> ModuleDependencies;
+	mutable TArray<FConvertedAssetRecord> ConvertedAssets;
 
-	UPROPERTY()
-	TArray<FConvertedAssetRecord> ConvertedAssets;
+	/**  */
+	TMap<FString, int32> RecordMap;
 };
diff --git a/Engine/Source/Developer/BlueprintNativeCodeGen/Private/BlueprintNativeCodeGenUtils.cpp b/Engine/Source/Developer/BlueprintNativeCodeGen/Private/BlueprintNativeCodeGenUtils.cpp
--- a/Engine/Source/Developer/BlueprintNativeCodeGen/Private/BlueprintNativeCodeGenUtils.cpp
+++ b/Engine/Source/Developer/BlueprintNativeCodeGen/Private/BlueprintNativeCodeGenUtils.cpp
@@ -2,6 +2,7 @@
 
 #include "BlueprintNativeCodeGenPCH.h"
 #include "BlueprintNativeCodeGenUtils.h"
+#include "NativeCodeGenCommandlineParams.h"
 #include "BlueprintNativeCodeGenCoordinator.h"
 #include "Kismet2/KismetReinstanceUtilities.h"	// for FBlueprintCompileReinstancer
 #include "Kismet2/CompilerResultsLog.h" 
@@ -11,6 +12,9 @@
 #include "Engine/UserDefinedEnum.h"
 #include "Kismet2/KismetEditorUtilities.h"		// for CompileBlueprint()
 #include "OutputDevice.h"						// for GWarn
+#include "GameProjectUtils.h"					// for GenerateGameModuleBuildFile
 
+DEFINE_LOG_CATEGORY(LogBlueprintCodeGen)
+
 /*******************************************************************************
  * BlueprintNativeCodeGenUtilsImpl
@@ -18,65 +22,28 @@
 
 namespace BlueprintNativeCodeGenUtilsImpl
 {
-	
-}
-
-/*******************************************************************************
- * FScopedFeedbackContext
- ******************************************************************************/
-
-//------------------------------------------------------------------------------
-FBlueprintNativeCodeGenUtils::FScopedFeedbackContext::FScopedFeedbackContext()
-	: OldContext(GWarn)
-	, ErrorCount(0)
-	, WarningCount(0)
-{
-	TreatWarningsAsErrors = GWarn->TreatWarningsAsErrors;
-	GWarn = this;
+	static bool WipeTargetPaths(const TArray<FString>& TargetPaths);
 }
 
 //------------------------------------------------------------------------------
-FBlueprintNativeCodeGenUtils::FScopedFeedbackContext::~FScopedFeedbackContext()
+static bool BlueprintNativeCodeGenUtilsImpl::WipeTargetPaths(const TArray<FString>& TargetPaths)
 {
-	GWarn = OldContext;
-}
+	IFileManager& FileManager = IFileManager::Get();
 
-//------------------------------------------------------------------------------
-bool FBlueprintNativeCodeGenUtils::FScopedFeedbackContext::HasErrors()
-{
-	return (ErrorCount > 0) || (TreatWarningsAsErrors && (WarningCount > 0));
-}
-
-//------------------------------------------------------------------------------
-void FBlueprintNativeCodeGenUtils::FScopedFeedbackContext::Serialize(const TCHAR* V, ELogVerbosity::Type Verbosity, const class FName& Category)
-{
-	switch (Verbosity)
+	bool bSuccess = true;
+	for (const FString& Path : TargetPaths)
 	{
-	case ELogVerbosity::Warning:
+		if (FileManager.FileExists(*Path))
 		{
-			++WarningCount;
-		} 
-		break;
-
-	case ELogVerbosity::Error:
-	case ELogVerbosity::Fatal:
+			bSuccess &= FileManager.Delete(*Path);
+		}
+		else if (FileManager.DirectoryExists(*Path))
 		{
-			++ErrorCount;
-		} 
-		break;
-
-	default:
-		break;
+			bSuccess &= FileManager.DeleteDirectory(*Path, /*RequireExists =*/false, /*Tree =*/true);
+		}
 	}
 
-	OldContext->Serialize(V, Verbosity, Category);
-}
-
-//------------------------------------------------------------------------------
-void FBlueprintNativeCodeGenUtils::FScopedFeedbackContext::Flush()
-{
-	WarningCount = ErrorCount = 0;
-	OldContext->Flush();
+	return bSuccess;
 }
 
 /*******************************************************************************
@@ -84,27 +51,49 @@
  ******************************************************************************/
 
 //------------------------------------------------------------------------------
-bool FBlueprintNativeCodeGenUtils::GenerateCodeModule(FBlueprintNativeCodeGenCoordinator& Coordinator)
+bool FBlueprintNativeCodeGenUtils::GenerateCodeModule(const FNativeCodeGenCommandlineParams& CommandParams)
 {
-	FBlueprintNativeCodeGenManifest& Manifest = Coordinator.Manifest;
+	FScopedFeedbackContext ScopedErrorTracker;
+	FBlueprintNativeCodeGenCoordinator Coordinator(CommandParams);
 
+	if (ScopedErrorTracker.HasErrors())
+	{
+		// creating the coordinator/manifest produced an error... do not carry on!
+		return false;
+	}
+	const FBlueprintNativeCodeGenManifest& Manifest = Coordinator.GetManifest();
+
+	if (CommandParams.bWipeRequested && !CommandParams.bPreviewRequested)
+	{
+		TArray<FString> TargetPaths = Manifest.GetTargetPaths();
+		if (!BlueprintNativeCodeGenUtilsImpl::WipeTargetPaths(TargetPaths))
+		{
+			UE_LOG(LogBlueprintCodeGen, Warning, TEXT("Failed to wipe target files/directories."));
+		}
+	}
+
 	TSharedPtr<FString> HeaderSource(new FString());
-	TSharedPtr<FString> CppSource(   new FString());
+	TSharedPtr<FString> CppSource(new FString());
 
-	bool bSuccess = (Coordinator.ConversionQueue.Num() > 0);
-	for (const FAssetData& Asset : Coordinator.ConversionQueue)
+	auto ConvertSingleAsset = [&Coordinator, &HeaderSource, &CppSource](FConvertedAssetRecord& ConversionRecord)->bool
 	{
 		FScopedFeedbackContext ScopedErrorTracker;
-		FConvertedAssetRecord& ConversionRecord = Manifest.CreateConversionRecord(Asset);
 
-		// loads the object if it is not already loaded
-		UObject* AssetObj = Asset.GetAsset();
-		FBlueprintNativeCodeGenUtils::GenerateCppCode(AssetObj, HeaderSource, CppSource);
+		UObject* AssetObj = ConversionRecord.AssetPtr.Get();
+		if (AssetObj != nullptr)
+		{
+			FBlueprintNativeCodeGenUtils::GenerateCppCode(AssetObj, HeaderSource, CppSource);
+		}
 
-		bSuccess &= !HeaderSource->IsEmpty() || !CppSource->IsEmpty();
+		bool bSuccess = !HeaderSource->IsEmpty() || !CppSource->IsEmpty();
 		if (!HeaderSource->IsEmpty())
 		{
-			bSuccess &= FFileHelper::SaveStringToFile(*HeaderSource, *ConversionRecord.GeneratedHeaderPath);
+			if (!FFileHelper::SaveStringToFile(*HeaderSource, *ConversionRecord.GeneratedHeaderPath))
+			{
+				bSuccess &= false;
+				ConversionRecord.GeneratedHeaderPath.Empty();
+			}
+			HeaderSource->Empty(HeaderSource->Len());
 		}
 		else
 		{
@@ -113,22 +102,30 @@
 
 		if (!CppSource->IsEmpty())
 		{
-			bSuccess &= FFileHelper::SaveStringToFile(*CppSource, *ConversionRecord.GeneratedCppPath);
+			if (!FFileHelper::SaveStringToFile(*CppSource, *ConversionRecord.GeneratedCppPath))
+			{
+				bSuccess &= false;
+				ConversionRecord.GeneratedCppPath.Empty();
+			}
+			CppSource->Empty(CppSource->Len());
 		}
 		else
 		{
 			ConversionRecord.GeneratedCppPath.Empty();
 		}
+		return bSuccess && !ScopedErrorTracker.HasErrors();
+	};
 
-		if (ScopedErrorTracker.HasErrors())
-		{
-			bSuccess = false;
-			break;
-		}
+	FBlueprintNativeCodeGenCoordinator::FConversionDelegate ConversionDelegate = FBlueprintNativeCodeGenCoordinator::FConversionDelegate::CreateLambda(ConvertSingleAsset);
+	if (CommandParams.bPreviewRequested)
+	{
+		ConversionDelegate.BindLambda([](FConvertedAssetRecord& ConversionRecord){ return true; });
 	}
-	Coordinator.Manifest.Save();
 
-	return bSuccess;
+	const bool bSuccess = Coordinator.ProcessConversionQueue(ConversionDelegate);
+	// save the manifest regardless of success (so we can get an idea of how successful it was)
+	Manifest.Save();
+	return bSuccess & !ScopedErrorTracker.HasErrors();
 }
 
 //------------------------------------------------------------------------------
@@ -188,3 +185,61 @@
 		ensure(false);
 	}
 }
+
+/*******************************************************************************
+ * FScopedFeedbackContext
+ ******************************************************************************/
+
+//------------------------------------------------------------------------------
+FBlueprintNativeCodeGenUtils::FScopedFeedbackContext::FScopedFeedbackContext()
+	: OldContext(GWarn)
+	, ErrorCount(0)
+	, WarningCount(0)
+{
+	TreatWarningsAsErrors = GWarn->TreatWarningsAsErrors;
+	GWarn = this;
+}
+
+//------------------------------------------------------------------------------
+FBlueprintNativeCodeGenUtils::FScopedFeedbackContext::~FScopedFeedbackContext()
+{
+	GWarn = OldContext;
+}
+
+//------------------------------------------------------------------------------
+bool FBlueprintNativeCodeGenUtils::FScopedFeedbackContext::HasErrors()
+{
+	return (ErrorCount > 0) || (TreatWarningsAsErrors && (WarningCount > 0));
+}
+
+//------------------------------------------------------------------------------
+void FBlueprintNativeCodeGenUtils::FScopedFeedbackContext::Serialize(const TCHAR* V, ELogVerbosity::Type Verbosity, const class FName& Category)
+{
+	switch (Verbosity)
+	{
+	case ELogVerbosity::Warning:
+		{
+			++WarningCount;
+		} 
+		break;
+
+	case ELogVerbosity::Error:
+	case ELogVerbosity::Fatal:
+		{
+			++ErrorCount;
+		} 
+		break;
+
+	default:
+		break;
+	}
+
+	OldContext->Serialize(V, Verbosity, Category);
+}
+
+//------------------------------------------------------------------------------
+void FBlueprintNativeCodeGenUtils::FScopedFeedbackContext::Flush()
+{
+	WarningCount = ErrorCount = 0;
+	OldContext->Flush();
+}
diff --git a/Engine/Source/Developer/BlueprintNativeCodeGen/Private/BlueprintNativeCodeGenUtils.h b/Engine/Source/Developer/BlueprintNativeCodeGen/Private/BlueprintNativeCodeGenUtils.h
--- a/Engine/Source/Developer/BlueprintNativeCodeGen/Private/BlueprintNativeCodeGenUtils.h
+++ b/Engine/Source/Developer/BlueprintNativeCodeGen/Private/BlueprintNativeCodeGenUtils.h
@@ -5,9 +5,11 @@
 #include "FeedbackContext.h"
 
 // Forward declares
-struct FBlueprintNativeCodeGenCoordinator;
+struct FNativeCodeGenCommandlineParams;
 class  SBuildProgressWidget;
 
+DECLARE_LOG_CATEGORY_EXTERN(LogBlueprintCodeGen, Log, All);
+
 /**  */
 struct FBlueprintNativeCodeGenUtils
 {
@@ -18,7 +20,7 @@ public:
 	 * @param  Coordinator    
 	 * @return 
 	 */
-	static bool GenerateCodeModule(FBlueprintNativeCodeGenCoordinator& Coordinator);
+	static bool GenerateCodeModule(const FNativeCodeGenCommandlineParams& CommandParams);
 
 	/**
 	 * Recompiles the bytecode of a blueprint only. Should only be run for
diff --git a/Engine/Source/Developer/BlueprintNativeCodeGen/Private/GenerateBlueprintCodeModuleCommandlet.cpp b/Engine/Source/Developer/BlueprintNativeCodeGen/Private/GenerateBlueprintCodeModuleCommandlet.cpp
--- a/Engine/Source/Developer/BlueprintNativeCodeGen/Private/GenerateBlueprintCodeModuleCommandlet.cpp
+++ b/Engine/Source/Developer/BlueprintNativeCodeGen/Private/GenerateBlueprintCodeModuleCommandlet.cpp
@@ -5,8 +5,8 @@
 #include "NativeCodeGenCommandlineParams.h"
 #include "BlueprintNativeCodeGenCoordinator.h"
 #include "BlueprintNativeCodeGenUtils.h"
-
-DEFINE_LOG_CATEGORY_STATIC(LogBlueprintCodeGen, Log, All);
+#include "BlueprintNativeCodeGenManifest.h"
+#include "FileManager.h"
 
 /*******************************************************************************
  * GenerateBlueprintCodeModuleImpl
@@ -14,7 +14,6 @@ DEFINE_LOG_CATEGORY_STATIC(LogBlueprintCodeGen, Log, All);
 
 namespace GenerateBlueprintCodeModuleImpl
 {
-
 }
 
 /*******************************************************************************
@@ -42,11 +41,8 @@ int32 UGenerateBlueprintCodeModuleCommandlet::Main(FString const& Params)
 	}
 
 	FBlueprintNativeCodeGenUtils::FScopedFeedbackContext ScopedErrorTracker;
-	FBlueprintNativeCodeGenCoordinator Coordinator(CommandlineParams);
-
-	if (!ScopedErrorTracker.HasErrors())
 	{
-		FBlueprintNativeCodeGenUtils::GenerateCodeModule(Coordinator);
+		FBlueprintNativeCodeGenUtils::GenerateCodeModule(CommandlineParams);
 	}
 	return ScopedErrorTracker.HasErrors();
 }
diff --git a/Engine/Source/Developer/BlueprintNativeCodeGen/Private/NativeCodeGenCommandlineParams.cpp b/Engine/Source/Developer/BlueprintNativeCodeGen/Private/NativeCodeGenCommandlineParams.cpp
--- a/Engine/Source/Developer/BlueprintNativeCodeGen/Private/NativeCodeGenCommandlineParams.cpp
+++ b/Engine/Source/Developer/BlueprintNativeCodeGen/Private/NativeCodeGenCommandlineParams.cpp
@@ -39,6 +39,9 @@ const FString FNativeCodeGenCommandlineParams::HelpMessage = TEXT("\n\
                         project's intermeadiate folder (within a sub-folder of \n\
                         its own). If specified as a relative path, it will be  \n\
                         relative to the target project's root directory.       \n\
+\n\
+    -moduleName=<Name>  Defines the output module's name. If left unset, then  \n\
+                        the module will default to... \n\
 \n\
     -wipe               Will clear out the target directory when specified.    \n\
                         Normally, without this switch, if the target directory \n\
@@ -53,6 +56,8 @@ const FString FNativeCodeGenCommandlineParams::HelpMessage = TEXT("\n\
                         will be determined by the old manifest file). If left  \n\
                         unset, then it will default to the module's target     \n\
                         directory.                                             \n\
+\n\
+    -preview            \n\
 \n\
     -help, -h, -?       Display this message and then exit.                    \n\
 \n");
@@ -63,7 +68,7 @@ const FString FNativeCodeGenCommandlineParams::HelpMessage = TEXT("\n\
 
 namespace NativeCodeGenCommandlineParamsImpl
 {
-	static const FString DefaultModuleName(TEXT("GeneratedBlueprintCode"));
+	static const FString DefaultModuleName(TEXT("BpCodeGen"));
 	static const FString ParamListDelim(TEXT(","));
 }
  
@@ -75,6 +80,7 @@ namespace NativeCodeGenCommandlineParamsImpl
 FNativeCodeGenCommandlineParams::FNativeCodeGenCommandlineParams(const TArray<FString>& CommandlineSwitches)
 	: bHelpRequested(false)
 	, bWipeRequested(false)
+	, bPreviewRequested(false)
 {
 	IFileManager& FileManager = IFileManager::Get();
 	const FString DefaultModulePath = FPaths::Combine(*FPaths::GameIntermediateDir(), *NativeCodeGenCommandlineParamsImpl::DefaultModuleName);
@@ -121,6 +127,10 @@ FNativeCodeGenCommandlineParams::FNativeCodeGenCommandlineParams(const TArray<FS
 		{
 			bWipeRequested = true;
 		}
+		else if (!Switch.Compare(TEXT("preview"), ESearchCase::IgnoreCase))
+		{
+			bPreviewRequested = true;
+		}
 		else if (!Switch.Compare(TEXT("manifest"), ESearchCase::IgnoreCase))
 		{
 			FPaths::NormalizeDirectoryName(Value);
diff --git a/Engine/Source/Developer/BlueprintNativeCodeGen/Private/NativeCodeGenCommandlineParams.h b/Engine/Source/Developer/BlueprintNativeCodeGen/Private/NativeCodeGenCommandlineParams.h
--- a/Engine/Source/Developer/BlueprintNativeCodeGen/Private/NativeCodeGenCommandlineParams.h
+++ b/Engine/Source/Developer/BlueprintNativeCodeGen/Private/NativeCodeGenCommandlineParams.h
@@ -16,7 +16,9 @@ public:
 	TArray<FString> WhiteListedAssetPaths;
 	TArray<FString> BlackListedAssetPaths;
 	FString         ModuleOutputDir;
+	FString			ModuleName;
 	FString			ManifestFilePath;
 	bool            bWipeRequested : 1;
+	bool            bPreviewRequested : 1;
 };
 
