Fixing CIS warnings: shadowed variables.

diff --git a/Engine/Source/Developer/BlueprintNativeCodeGen/Private/BlueprintNativeCodeGenCoordinator.cpp b/Engine/Source/Developer/BlueprintNativeCodeGen/Private/BlueprintNativeCodeGenCoordinator.cpp
--- a/Engine/Source/Developer/BlueprintNativeCodeGen/Private/BlueprintNativeCodeGenCoordinator.cpp
+++ b/Engine/Source/Developer/BlueprintNativeCodeGen/Private/BlueprintNativeCodeGenCoordinator.cpp
@@ -354,10 +354,10 @@ static void BlueprintNativeCodeGenCoordinatorImpl::GatherInnerDependencies(UObje
 	if (UStruct* StructAsset = GetAssetStruct(AssetObj))
 	{
 		IAssetRegistry& AssetRegistry = GetAssetRegistry();
-		auto AddDependencyForConversion = [&AssetsOut, &AssetRegistry](UObject* AssetObj)->bool
+		auto AddDependencyForConversion = [&AssetsOut, &AssetRegistry](UObject* Object)->bool
 		{
-			ensure(AssetObj->IsAsset());
-			FAssetData AssetInfo = AssetRegistry.GetAssetByObjectPath(*AssetObj->GetPathName());
+			ensure(Object->IsAsset());
+			FAssetData AssetInfo = AssetRegistry.GetAssetByObjectPath(*Object->GetPathName());
 
 			const bool bAssetFound = AssetInfo.IsValid();
 			if (bAssetFound)
@@ -376,9 +376,9 @@ static void BlueprintNativeCodeGenCoordinatorImpl::GatherInnerDependencies(UObje
 				AddDependencyForConversion(Blueprint);
 			}
 		}
-		for (UUserDefinedStruct* StructAsset : AssetDependencies.ConvertedStructs)
+		for (UUserDefinedStruct* StructDependency : AssetDependencies.ConvertedStructs)
 		{
-			AddDependencyForConversion(StructAsset);
+			AddDependencyForConversion(StructDependency);
 		}
 		for (UUserDefinedEnum* EnumAsset : AssetDependencies.ConvertedEnum)
 		{
diff --git a/Engine/Source/Developer/BlueprintNativeCodeGen/Private/BlueprintNativeCodeGenUtils.cpp b/Engine/Source/Developer/BlueprintNativeCodeGen/Private/BlueprintNativeCodeGenUtils.cpp
--- a/Engine/Source/Developer/BlueprintNativeCodeGen/Private/BlueprintNativeCodeGenUtils.cpp
+++ b/Engine/Source/Developer/BlueprintNativeCodeGen/Private/BlueprintNativeCodeGenUtils.cpp
@@ -77,7 +77,7 @@ bool FBlueprintNativeCodeGenUtils::GenerateCodeModule(const FNativeCodeGenComman
 
 	auto ConvertSingleAsset = [&Coordinator, &HeaderSource, &CppSource](FConvertedAssetRecord& ConversionRecord)->bool
 	{
-		FScopedFeedbackContext ScopedErrorTracker;
+		FScopedFeedbackContext NestedErrorTracker;
 
 		UObject* AssetObj = ConversionRecord.AssetPtr.Get();
 		if (AssetObj != nullptr)
@@ -113,7 +113,7 @@ bool FBlueprintNativeCodeGenUtils::GenerateCodeModule(const FNativeCodeGenComman
 		{
 			ConversionRecord.GeneratedCppPath.Empty();
 		}
-		return bSuccess && !ScopedErrorTracker.HasErrors();
+		return bSuccess && !NestedErrorTracker.HasErrors();
 	};
 
 	FBlueprintNativeCodeGenCoordinator::FConversionDelegate ConversionDelegate = FBlueprintNativeCodeGenCoordinator::FConversionDelegate::CreateLambda(ConvertSingleAsset);
