Fixing the in-editor Blueprint conversion tool so that the files generated match the specified includes.

diff --git a/Engine/Source/Developer/BlueprintCompilerCppBackend/Private/BlueprintCompilerCppBackendBase.cpp b/Engine/Source/Developer/BlueprintCompilerCppBackend/Private/BlueprintCompilerCppBackendBase.cpp
--- a/Engine/Source/Developer/BlueprintCompilerCppBackend/Private/BlueprintCompilerCppBackendBase.cpp
+++ b/Engine/Source/Developer/BlueprintCompilerCppBackend/Private/BlueprintCompilerCppBackendBase.cpp
@@ -3,6 +3,7 @@
 #include "BlueprintCompilerCppBackendModulePrivatePCH.h"
 #include "BlueprintCompilerCppBackendBase.h"
 #include "BlueprintCompilerCppBackendUtils.h"
+#include "IBlueprintCompilerCppBackendModule.h" // for GetBaseFilename()
 
 void FBlueprintCompilerCppBackendBase::EmitStructProperties(FStringOutputDevice& Target, UStruct* SourceClass)
 {
@@ -420,15 +421,18 @@ void FBlueprintCompilerCppBackendBase::EmitFileBeginning(const FString& CleanNam
 				check(Field);
 				const bool bWantedType = Field->IsA<UBlueprintGeneratedClass>() || Field->IsA<UUserDefinedEnum>() || Field->IsA<UUserDefinedStruct>();
 
-				// Wanted no-native type, thet will be converted
+				// Wanted no-native type, that will be converted
 				if (bWantedType)
 				{
+					// @TODO: Need to query if this asset will actually be converted
+
 					const FString Name = Field->GetName();
 					bool bAlreadyIncluded = false;
 					AlreadyIncluded.Add(Name, &bAlreadyIncluded);
 					if (!bAlreadyIncluded)
 					{
-						EmitIncludeHeader(Dst, *Name, true);
+						const FString GeneratedFilename = IBlueprintCompilerCppBackendModule::GetBaseFilename(Field);
+						EmitIncludeHeader(Dst, *GeneratedFilename, true);
 					}
 				}
 				// headers for native items
diff --git a/Engine/Source/Developer/BlueprintNativeCodeGen/Private/NativeCodeGenerationTool.cpp b/Engine/Source/Developer/BlueprintNativeCodeGen/Private/NativeCodeGenerationTool.cpp
--- a/Engine/Source/Developer/BlueprintNativeCodeGen/Private/NativeCodeGenerationTool.cpp
+++ b/Engine/Source/Developer/BlueprintNativeCodeGen/Private/NativeCodeGenerationTool.cpp
@@ -13,6 +13,7 @@
 #include "BlueprintNativeCodeGenUtils.h"
 //#include "Editor/KismetCompiler/Public/BlueprintCompilerCppBackendInterface.h"
 #include "Developer/BlueprintCompilerCppBackend/Public/BlueprintCompilerCppBackendGatherDependencies.h"
+#include "IBlueprintCompilerCppBackendModule.h" // for GetBaseFilename()
 
 #define LOCTEXT_NAMESPACE "NativeCodeGenerationTool"
 
@@ -29,6 +30,7 @@ struct FGeneratedCodeData
 		FName GeneratedClassName, SkeletonClassName;
 		InBlueprint.GetBlueprintClassNames(GeneratedClassName, SkeletonClassName);
 		ClassName = GeneratedClassName.ToString();
+		BaseFilename = IBlueprintCompilerCppBackendModule::GetBaseFilename(&InBlueprint);
 
 		GatherUserDefinedDependencies(InBlueprint);
 	}
@@ -36,6 +38,7 @@ struct FGeneratedCodeData
 	FString TypeDependencies;
 	FString ErrorString;
 	FString ClassName;
+	FString BaseFilename;
 	TWeakObjectPtr<UBlueprint> Blueprint;
 	TSet<UField*> DependentObjects;
 
@@ -82,12 +85,12 @@ struct FGeneratedCodeData
 
 	FString HeaderFileName() const
 	{
-		return ClassName + TEXT(".h");
+		return BaseFilename + TEXT(".h");
 	}
 
 	FString SourceFileName() const
 	{
-		return ClassName + TEXT(".cpp");
+		return BaseFilename + TEXT(".cpp");
 	}
 
 	bool Save(const FString& HeaderDirPath, const FString& CppDirPath)
@@ -111,7 +114,9 @@ struct FGeneratedCodeData
 			FBlueprintNativeCodeGenUtils::GenerateCppCode(Obj, HeaderSource, CppSource);
 			SlowTask.EnterProgressFrame();
 
-			const FString FullHeaderFilename = FPaths::Combine(*HeaderDirPath, *(Obj->GetName() + TEXT(".h")));
+			const FString BaseFilename = IBlueprintCompilerCppBackendModule::GetBaseFilename(Obj);
+
+			const FString FullHeaderFilename = FPaths::Combine(*HeaderDirPath, *(BaseFilename + TEXT(".h")));
 			const bool bHeaderSaved = FFileHelper::SaveStringToFile(*HeaderSource, *FullHeaderFilename);
 			if (!bHeaderSaved)
 			{
@@ -125,7 +130,7 @@ struct FGeneratedCodeData
 			SlowTask.EnterProgressFrame();
 			if (!CppSource->IsEmpty())
 			{
-				const FString NewCppFilename = FPaths::Combine(*CppDirPath, *(Obj->GetName() + TEXT(".cpp")));
+				const FString NewCppFilename = FPaths::Combine(*CppDirPath, *(BaseFilename + TEXT(".cpp")));
 				const bool bCppSaved = FFileHelper::SaveStringToFile(*CppSource, *NewCppFilename);
 				if (!bCppSaved)
 				{
