Fix CIS warning on Mac #buildwatcher

diff --git a/Engine/Source/Developer/BlueprintNativeCodeGen/Private/NativeCodeGenerationTool.cpp b/Engine/Source/Developer/BlueprintNativeCodeGen/Private/NativeCodeGenerationTool.cpp
--- a/Engine/Source/Developer/BlueprintNativeCodeGen/Private/NativeCodeGenerationTool.cpp
+++ b/Engine/Source/Developer/BlueprintNativeCodeGen/Private/NativeCodeGenerationTool.cpp
@@ -114,9 +114,9 @@ struct FGeneratedCodeData
 			FBlueprintNativeCodeGenUtils::GenerateCppCode(Obj, HeaderSource, CppSource);
 			SlowTask.EnterProgressFrame();
 
-			const FString BaseFilename = IBlueprintCompilerCppBackendModule::GetBaseFilename(Obj);
+			const FString BackendBaseFilename = IBlueprintCompilerCppBackendModule::GetBaseFilename(Obj);
 
-			const FString FullHeaderFilename = FPaths::Combine(*HeaderDirPath, *(BaseFilename + TEXT(".h")));
+			const FString FullHeaderFilename = FPaths::Combine(*HeaderDirPath, *(BackendBaseFilename + TEXT(".h")));
 			const bool bHeaderSaved = FFileHelper::SaveStringToFile(*HeaderSource, *FullHeaderFilename);
 			if (!bHeaderSaved)
 			{
@@ -130,7 +130,7 @@ struct FGeneratedCodeData
 			SlowTask.EnterProgressFrame();
 			if (!CppSource->IsEmpty())
 			{
-				const FString NewCppFilename = FPaths::Combine(*CppDirPath, *(BaseFilename + TEXT(".cpp")));
+				const FString NewCppFilename = FPaths::Combine(*CppDirPath, *(BackendBaseFilename + TEXT(".cpp")));
 				const bool bCppSaved = FFileHelper::SaveStringToFile(*CppSource, *NewCppFilename);
 				if (!bCppSaved)
 				{
