(WIP) Blueprint conversion tool now generates its own Build.cs file (with the proper dependencies).

diff --git a/Engine/Source/Developer/BlueprintNativeCodeGen/Private/BlueprintNativeCodeGenCoordinator.cpp b/Engine/Source/Developer/BlueprintNativeCodeGen/Private/BlueprintNativeCodeGenCoordinator.cpp
--- a/Engine/Source/Developer/BlueprintNativeCodeGen/Private/BlueprintNativeCodeGenCoordinator.cpp
+++ b/Engine/Source/Developer/BlueprintNativeCodeGen/Private/BlueprintNativeCodeGenCoordinator.cpp
@@ -484,21 +484,39 @@ bool FBlueprintNativeCodeGenCoordinator::ProcessConversionQueue(const FConversio
 	ConversionQueryDelegate.BindLambda(IsAssetTargetedForConversion);
 
 	bool bSuccess = true;
-	for (int32 AssetIndex = 0; AssetIndex < ConversionQueue.Num(); ++AssetIndex)
+	// pre-process the conversion queue (has to happen before converting, 
+	// so we can gather module dependencies in FBlueprintNativeCodeGenManifest::LogDependencies);
+	// if this were to happen in the same loop as converting, then the compilation
+	// of one asset could result in a later asset having its linker/loader reset
+	// (before we processed it... meaning we wouldn't have an easy way to 
+	// determine that later Blueprint's module dependencies)
+	for (int32 AssetIndex = 0; AssetIndex < ConversionQueue.Num() && bSuccess; ++AssetIndex)
 	{
 		const FAssetData& AssetInfo = ConversionQueue[AssetIndex];
+
 		// load asset, add dependencies that can only be gathered from a 
 		// loaded asset (property types, etc.)
 		UObject* AssetObj = AssetInfo.GetAsset();
+		// NOTE: this will most likely grow the ConversionQueue as we iterate through it
 		BlueprintNativeCodeGenCoordinatorImpl::GatherInnerDependencies(AssetObj, ConversionQueue);
 
+		bSuccess &= Manifest.LogDependencies(AssetInfo);
+		if (!bSuccess)
+		{
+			UE_LOG(LogBlueprintCodeGen, Error, TEXT("Failed to pre-process '%s' for conversion"), *AssetInfo.AssetName.ToString());
+		}
+	}
+
+	for (int32 AssetIndex = 0; AssetIndex < ConversionQueue.Num() && bSuccess; ++AssetIndex)
+	{
+		const FAssetData& AssetInfo = ConversionQueue[AssetIndex];
+
 		FConvertedAssetRecord& ConversionRecord = Manifest.CreateConversionRecord(AssetInfo);
 		bSuccess &= ConversionDelegate.Execute(ConversionRecord);
 
 		if (!bSuccess)
 		{
 			UE_LOG(LogBlueprintCodeGen, Error, TEXT("Failed to fully convert '%s'"), *AssetInfo.AssetName.ToString());
-			break;
 		}
 	}
 
diff --git a/Engine/Source/Developer/BlueprintNativeCodeGen/Private/BlueprintNativeCodeGenManifest.cpp b/Engine/Source/Developer/BlueprintNativeCodeGen/Private/BlueprintNativeCodeGenManifest.cpp
--- a/Engine/Source/Developer/BlueprintNativeCodeGen/Private/BlueprintNativeCodeGenManifest.cpp
+++ b/Engine/Source/Developer/BlueprintNativeCodeGen/Private/BlueprintNativeCodeGenManifest.cpp
@@ -19,11 +19,13 @@ DEFINE_LOG_CATEGORY_STATIC(LogNativeCodeGenManifest, Log, All);
 namespace BlueprintNativeCodeGenManifestImpl
 {
 	static const int64 CPF_NoFlags = 0x00;
-	static const FString ManifestFileExt = TEXT(".bpgen.manifest");
-	static const FString CppFileExt	     = TEXT(".cpp");
-	static const FString HeaderFileExt   = TEXT(".h");
-	static const FString HeaderSubDir    = TEXT("Public");
-	static const FString CppSubDir       = TEXT("Private");
+	static const FString ManifestFileExt    = TEXT(".bpgen.manifest");
+	static const FString CppFileExt	        = TEXT(".cpp");
+	static const FString HeaderFileExt      = TEXT(".h");
+	static const FString HeaderSubDir       = TEXT("Public");
+	static const FString CppSubDir          = TEXT("Private");
+	static const FString ModuleBuildFileExt = TEXT(".Build.cs");
+	static const FString FallbackModuleName = TEXT("BpCodeGen");
 
 	/**  */
 	static bool LoadManifest(const FString& FilePath, FBlueprintNativeCodeGenManifest* Manifest);
@@ -37,6 +39,8 @@ namespace BlueprintNativeCodeGenManifestImpl
 	static FString GenerateCppSavePath(const FString& ModulePath, const FAssetData& Asset);
 
 	static FString GetComparibleDirPath(const FString& DirectoryPath);
+
+	static bool GatherModuleDependencies(const UObject* AssetObj, TArray<UPackage*>& DependenciesOut);
 }
 
 //------------------------------------------------------------------------------
@@ -110,6 +114,61 @@ static FString BlueprintNativeCodeGenManifestImpl::GetComparibleDirPath(const FS
 	return NormalizedPath;
 }
 
+//------------------------------------------------------------------------------
+static bool BlueprintNativeCodeGenManifestImpl::GatherModuleDependencies(const UObject* AssetObj, TArray<UPackage*>& DependenciesOut)
+{
+	UPackage* AssetPackage = AssetObj->GetOutermost();
+
+	const FLinkerLoad* PkgLinker = FLinkerLoad::FindExistingLinkerForPackage(AssetPackage);
+	const bool bFoundLinker = (PkgLinker != nullptr);
+
+	if (ensure(bFoundLinker))
+	{
+		for (const FObjectImport& PkgImport : PkgLinker->ImportMap)
+		{
+			if (PkgImport.ClassName != NAME_Package)
+			{
+				continue;
+			}
+
+			UPackage* DependentPackage = FindObject<UPackage>(/*Outer =*/nullptr, *PkgImport.ObjectName.ToString(), /*ExactClass =*/true);
+			if (DependentPackage == nullptr)
+			{
+				continue;
+			}
+
+			// we want only native packages, ones that are not editor-only
+			if ((DependentPackage->PackageFlags & (PKG_CompiledIn | PKG_EditorOnly)) == PKG_CompiledIn)
+			{
+				DependenciesOut.AddUnique(DependentPackage);// PkgImport.ObjectName.ToString());
+			}
+		}
+	}
+
+	return bFoundLinker;
+}
+
+/*******************************************************************************
+ * FConvertedAssetRecord
+ ******************************************************************************/
+ 
+//------------------------------------------------------------------------------
+FConvertedAssetRecord::FConvertedAssetRecord(const FAssetData& AssetInfo, const FString& TargetModulePath)
+	: AssetPtr(AssetInfo.GetAsset())
+	, AssetType(AssetInfo.GetClass())
+	, AssetPath(AssetInfo.PackageName.ToString())
+	, GeneratedHeaderPath(BlueprintNativeCodeGenManifestImpl::GenerateHeaderSavePath(TargetModulePath, AssetInfo))
+	, GeneratedCppPath(BlueprintNativeCodeGenManifestImpl::GenerateCppSavePath(TargetModulePath, AssetInfo))
+{
+}
+
+//------------------------------------------------------------------------------
+bool FConvertedAssetRecord::IsValid()
+{
+	// every conversion will have a header file (interfaces don't have implementation files)
+	return AssetPtr.IsValid() && !GeneratedHeaderPath.IsEmpty() && (AssetType != nullptr) && !AssetPath.IsEmpty();
+}
+
 /*******************************************************************************
  * FBlueprintNativeCodeGenManifest
  ******************************************************************************/
@@ -137,6 +196,7 @@ FBlueprintNativeCodeGenManifest::FBlueprintNativeCodeGenManifest(const FNativeCo
 	using namespace BlueprintNativeCodeGenManifestImpl;
 	bool const bLoadExistingManifest = !CommandlineParams.bWipeRequested || CommandlineParams.ModuleOutputDir.IsEmpty();
 	
+	ModuleName = CommandlineParams.ModuleName;
 	ModulePath = CommandlineParams.ModuleOutputDir;
 	if (FPaths::IsRelative(ModulePath))
 	{
@@ -149,9 +209,13 @@ FBlueprintNativeCodeGenManifest::FBlueprintNativeCodeGenManifest(const FNativeCo
 	}
 	else
 	{
+		if (!ensure(!ModuleName.IsEmpty()))
+		{
+			ModuleName = FallbackModuleName;
+		}
 		if (!ensure(!ModulePath.IsEmpty()))
 		{
-			ModulePath = FPaths::Combine(*FPaths::GameIntermediateDir(), TEXT("BpCodeGen"));
+			ModulePath = FPaths::Combine(*FPaths::GameIntermediateDir(), *ModuleName);
 		}
 		ManifestPath = FPaths::Combine(*GetModuleDir(), *GetDefaultFilename());
 	}
@@ -183,13 +247,27 @@ FBlueprintNativeCodeGenManifest::FBlueprintNativeCodeGenManifest(const FNativeCo
 			}
 		}
 
+		const bool bNewModuleNameRequested = !CommandlineParams.ModuleName.IsEmpty() && (CommandlineParams.ModuleName != ModuleName);
+		if (bNewModuleNameRequested && !CommandlineParams.bPreviewRequested) 
+		{
+			// delete the old module file (if one exists)
+			IFileManager::Get().Delete(*GetBuildFilePath());
+		}
+
 		// if we were only interested in obtaining the module path
 		if (CommandlineParams.bWipeRequested)
 		{
+			ModuleName = CommandlineParams.ModuleName;
 			Clear();
 		}
 		else
 		{
+			if (bNewModuleNameRequested)
+			{
+				UE_LOG(LogNativeCodeGenManifest, Warning, TEXT("The specified module name (%s) doesn't match the existing one (%s). Overridding with the new name."),
+					*CommandlineParams.ModuleName, *ModuleName);
+				ModuleName = CommandlineParams.ModuleName;
+			}
 			MapConvertedAssets();
 		}
 	}
@@ -199,22 +277,43 @@ FBlueprintNativeCodeGenManifest::FBlueprintNativeCodeGenManifest(const FNativeCo
 FConvertedAssetRecord& FBlueprintNativeCodeGenManifest::CreateConversionRecord(const FAssetData& AssetInfo)
 {
 	const FString AssetPath = AssetInfo.PackageName.ToString();
-	if (FConvertedAssetRecord* ExistingConversionRecord = FindConversionRecord(AssetPath))
+	const FString TargetModulePath = GetModuleDir();
+	// load the asset (if it isn't already)
+	const UObject* AssetObj = AssetInfo.GetAsset();
+
+	FConvertedAssetRecord* ConversionRecordPtr = FindConversionRecord(AssetPath);
+	if (ConversionRecordPtr == nullptr)
 	{
-		return *ExistingConversionRecord;
+		ConversionRecordPtr = &ConvertedAssets[ ConvertedAssets.Add(FConvertedAssetRecord(AssetInfo, TargetModulePath)) ];
 	}
+	else if ( !ensure((ConversionRecordPtr->AssetPath == AssetPath) && (ConversionRecordPtr->AssetType == AssetInfo.GetClass())) )
+	{
+		UE_LOG(LogNativeCodeGenManifest, Warning, 
+			TEXT("The existing manifest entery for '%s' doesn't match what was expected (asset path and/or type). Updating it to match the asset."),
+			*AssetPath);
 
-	const FString TargetModulePath = GetModuleDir();
-
-	FConvertedAssetRecord NewConversionRecord;
-	NewConversionRecord.AssetPtr  = AssetInfo.GetAsset();
-	NewConversionRecord.AssetType = AssetInfo.GetClass();
-	NewConversionRecord.AssetPath = AssetPath;
-	NewConversionRecord.GeneratedHeaderPath = BlueprintNativeCodeGenManifestImpl::GenerateHeaderSavePath(TargetModulePath, AssetInfo);
-	NewConversionRecord.GeneratedCppPath    = BlueprintNativeCodeGenManifestImpl::GenerateCppSavePath(TargetModulePath, AssetInfo);
+		ConversionRecordPtr->AssetPath = AssetPath;
+		ConversionRecordPtr->AssetType = AssetInfo.GetClass();
+	}
+	
+	FConvertedAssetRecord& ConversionRecord = *ConversionRecordPtr;
+	if (!ConversionRecord.IsValid())
+	{
+		// if this was a manifest entry that was loaded from an existing file, 
+		// then it wouldn't have the asset pointer (which is transient data)
+		if (!ConversionRecord.AssetPtr.IsValid())
+		{
+			ConversionRecord.AssetPtr = AssetObj;
+		}
+		if (ConversionRecord.GeneratedHeaderPath.IsEmpty())
+		{
+			ConversionRecord.GeneratedHeaderPath = BlueprintNativeCodeGenManifestImpl::GenerateHeaderSavePath(TargetModulePath, AssetInfo);
+			ConversionRecord.GeneratedCppPath    = BlueprintNativeCodeGenManifestImpl::GenerateCppSavePath(TargetModulePath, AssetInfo);
+		}
+		ensure(ConversionRecord.IsValid());
+	}
 
-	RecordMap.Add(AssetPath, ConvertedAssets.Num());
-	return ConvertedAssets[ ConvertedAssets.Add(NewConversionRecord) ];
+	return ConversionRecord;
 }
 
 //------------------------------------------------------------------------------
@@ -240,6 +339,15 @@ FConvertedAssetRecord* FBlueprintNativeCodeGenManifest::FindConversionRecord(con
 	return FoundRecord;
 }
 
+//------------------------------------------------------------------------------
+bool FBlueprintNativeCodeGenManifest::LogDependencies(const FAssetData& AssetInfo)
+{
+	// load the asset (if it isn't already)
+	const UObject* AssetObj = AssetInfo.GetAsset();
+
+	return BlueprintNativeCodeGenManifestImpl::GatherModuleDependencies(AssetObj, ModuleDependencies);
+}
+
 //------------------------------------------------------------------------------
 TArray<FString> FBlueprintNativeCodeGenManifest::GetTargetPaths() const 
 {
@@ -249,11 +357,17 @@ TArray<FString> FBlueprintNativeCodeGenManifest::GetTargetPaths() const
 	DestPaths.Add(ManifestPath);
 	DestPaths.Add(BlueprintNativeCodeGenManifestImpl::GetHeaderSaveDir(TargetModulePath));
 	DestPaths.Add(BlueprintNativeCodeGenManifestImpl::GetCppSaveDir(TargetModulePath));
-	// @TODO: Add *Build.cs file
+	DestPaths.Add(GetBuildFilePath());
 
 	return DestPaths;
 }
 
+//------------------------------------------------------------------------------
+FString FBlueprintNativeCodeGenManifest::GetBuildFilePath() const
+{
+	return FPaths::Combine(*GetModuleDir(), *ModuleName) + BlueprintNativeCodeGenManifestImpl::ModuleBuildFileExt;
+}
+
 //------------------------------------------------------------------------------
 bool FBlueprintNativeCodeGenManifest::Save(bool bSort) const
 {
@@ -314,15 +428,6 @@ void FBlueprintNativeCodeGenManifest::MapConvertedAssets()
 	{
 		const FConvertedAssetRecord& AssetRecord = ConvertedAssets[RecordIndex];
 		RecordMap.Add(AssetRecord.AssetPath, RecordIndex);
-		
-// 		if (UObject* AssetObj = AssetRecord.AssetPtr.Get())
-// 		{
-//			RecordMap.Add(AssetRecord.AssetPath, RecordIndex);
-// 		}
-// 		else
-// 		{
-// 			UE_LOG(LogBlueprintCodeGen, Warning, TEXT("Invalid conversion record found for: %s"), *AssetRecord.AssetPath);
-// 		}
 	}
 }
 
diff --git a/Engine/Source/Developer/BlueprintNativeCodeGen/Private/BlueprintNativeCodeGenManifest.h b/Engine/Source/Developer/BlueprintNativeCodeGen/Private/BlueprintNativeCodeGenManifest.h
--- a/Engine/Source/Developer/BlueprintNativeCodeGen/Private/BlueprintNativeCodeGenManifest.h
+++ b/Engine/Source/Developer/BlueprintNativeCodeGen/Private/BlueprintNativeCodeGenManifest.h
@@ -17,6 +17,12 @@ struct FConvertedAssetRecord
 {
 	GENERATED_USTRUCT_BODY()
 
+public:
+	FConvertedAssetRecord() {}
+	FConvertedAssetRecord(const FAssetData& AssetInfo, const FString& TargetModulePath);
+
+	bool IsValid();
+
 public:
 	TWeakObjectPtr<UObject> AssetPtr;
 
@@ -49,12 +55,17 @@ public:
 	FBlueprintNativeCodeGenManifest(const FNativeCodeGenCommandlineParams& CommandlineParams);
 
 	/**
+	 * A query that retrieves a list of file/directory paths. These paths are 
+	 * the destination targets for files and directories that should/will be 
+	 * generated as part of the conversion process.
 	 * 
-	 * 
-	 * @return 
+	 * @return An array of file/directory paths that serve as destinations for the conversion process.
 	 */
 	TArray<FString> GetTargetPaths() const;
 
+	/**  */
+	FString GetBuildFilePath() const;
+
 	/**
 	 * 
 	 * 
@@ -71,6 +82,12 @@ public:
 	 */
 	FConvertedAssetRecord* FindConversionRecord(const FString& AssetPath, bool bSlowSearch = false);
 
+	/**  */
+	bool LogDependencies(const FAssetData& AssetInfo);
+
+	/**  */
+	const TArray<UPackage*>& GetModuleDependencies() const { return ModuleDependencies; }
+
 	/**
 	 * 
 	 */
@@ -95,10 +112,16 @@ private:
 	/**  */
 	FString ManifestPath;
 
+	UPROPERTY()
+	FString ModuleName;
+
 	/** Relative to the project's directory */
 	UPROPERTY()
 	FString ModulePath;
 
+	UPROPERTY()
+	TArray<UPackage*> ModuleDependencies;
+
 	/** Mutable so Save() can sort the array when requested */
 	UPROPERTY()
 	mutable TArray<FConvertedAssetRecord> ConvertedAssets;
diff --git a/Engine/Source/Developer/BlueprintNativeCodeGen/Private/BlueprintNativeCodeGenUtils.cpp b/Engine/Source/Developer/BlueprintNativeCodeGen/Private/BlueprintNativeCodeGenUtils.cpp
--- a/Engine/Source/Developer/BlueprintNativeCodeGen/Private/BlueprintNativeCodeGenUtils.cpp
+++ b/Engine/Source/Developer/BlueprintNativeCodeGen/Private/BlueprintNativeCodeGenUtils.cpp
@@ -13,6 +13,8 @@
 #include "Kismet2/KismetEditorUtilities.h"		// for CompileBlueprint()
 #include "OutputDevice.h"						// for GWarn
 #include "GameProjectUtils.h"					// for GenerateGameModuleBuildFile
+#include "Editor/GameProjectGeneration/Public/GameProjectUtils.h" // for GenerateGameModuleBuildFile()
+#include "App.h"								// for GetGameName()
 
 DEFINE_LOG_CATEGORY(LogBlueprintCodeGen)
 
@@ -23,6 +25,7 @@ DEFINE_LOG_CATEGORY(LogBlueprintCodeGen)
 namespace BlueprintNativeCodeGenUtilsImpl
 {
 	static bool WipeTargetPaths(const TArray<FString>& TargetPaths);
+	static bool GenerateModuleBuildFile(const FString& ModuleName, const FBlueprintNativeCodeGenManifest& Manifest);
 }
 
 //------------------------------------------------------------------------------
@@ -46,6 +49,50 @@ static bool BlueprintNativeCodeGenUtilsImpl::WipeTargetPaths(const TArray<FStrin
 	return bSuccess;
 }
 
+//------------------------------------------------------------------------------
+static bool BlueprintNativeCodeGenUtilsImpl::GenerateModuleBuildFile(const FString& ModuleName, const FBlueprintNativeCodeGenManifest& Manifest)
+{
+	FModuleManager& ModuleManager = FModuleManager::Get();
+	
+	TArray<FString> PublicDependencies;
+	if (GameProjectUtils::ProjectHasCodeFiles()) 
+	{
+		const FString GameModuleName = FApp::GetGameName();
+		if (ModuleManager.ModuleExists(*GameModuleName))
+		{
+			PublicDependencies.Add(GameModuleName);
+		}
+	}
+
+	TArray<FString> PrivateDependencies;
+
+	const TArray<UPackage*>& ModulePackages = Manifest.GetModuleDependencies();
+	PrivateDependencies.Reserve(ModulePackages.Num());
+
+	for (UPackage* ModulePkg : ModulePackages)
+	{
+		const FString ModuleName = FPackageName::GetLongPackageAssetName(ModulePkg->GetName());
+		if (ModuleManager.ModuleExists(*ModuleName))
+		{
+			PrivateDependencies.Add(ModuleName);
+		}
+		else
+		{
+			UE_LOG(LogBlueprintCodeGen, Warning, TEXT("Failed to find module for package: %s"), *ModuleName);
+		}
+	}
+
+	const FString BuildFilePath = Manifest.GetBuildFilePath();
+	FText ErrorMessage;
+	GameProjectUtils::GenerateGameModuleBuildFile(BuildFilePath, ModuleName, PublicDependencies, PrivateDependencies, ErrorMessage);
+
+	if (!ErrorMessage.IsEmpty())
+	{
+		UE_LOG(LogBlueprintCodeGen, Error, TEXT("Failed to generate module build file: %s"), *ErrorMessage.ToString());
+	}
+	return !ErrorMessage.IsEmpty();
+}
+
 /*******************************************************************************
  * FBlueprintNativeCodeGenUtils
  ******************************************************************************/
@@ -86,32 +133,34 @@ bool FBlueprintNativeCodeGenUtils::GenerateCodeModule(const FNativeCodeGenComman
 		}
 
 		bool bSuccess = !HeaderSource->IsEmpty() || !CppSource->IsEmpty();
-		if (!HeaderSource->IsEmpty())
+		// run the cpp first, because we cue off of the presence of a header for 
+		// a valid conversion record (see FConvertedAssetRecord::IsValid)
+		if (!CppSource->IsEmpty())
 		{
-			if (!FFileHelper::SaveStringToFile(*HeaderSource, *ConversionRecord.GeneratedHeaderPath))
+			if (!FFileHelper::SaveStringToFile(*CppSource, *ConversionRecord.GeneratedCppPath))
 			{
 				bSuccess &= false;
-				ConversionRecord.GeneratedHeaderPath.Empty();
+				ConversionRecord.GeneratedCppPath.Empty();
 			}
-			HeaderSource->Empty(HeaderSource->Len());
+			CppSource->Empty(CppSource->Len());
 		}
 		else
 		{
-			ConversionRecord.GeneratedHeaderPath.Empty();
+			ConversionRecord.GeneratedCppPath.Empty();
 		}
 
-		if (!CppSource->IsEmpty())
+		if (bSuccess && !HeaderSource->IsEmpty())
 		{
-			if (!FFileHelper::SaveStringToFile(*CppSource, *ConversionRecord.GeneratedCppPath))
+			if (!FFileHelper::SaveStringToFile(*HeaderSource, *ConversionRecord.GeneratedHeaderPath))
 			{
 				bSuccess &= false;
-				ConversionRecord.GeneratedCppPath.Empty();
+				ConversionRecord.GeneratedHeaderPath.Empty();
 			}
-			CppSource->Empty(CppSource->Len());
+			HeaderSource->Empty(HeaderSource->Len());
 		}
 		else
 		{
-			ConversionRecord.GeneratedCppPath.Empty();
+			ConversionRecord.GeneratedHeaderPath.Empty();
 		}
 		return bSuccess && !NestedErrorTracker.HasErrors();
 	};
@@ -123,6 +172,11 @@ bool FBlueprintNativeCodeGenUtils::GenerateCodeModule(const FNativeCodeGenComman
 	}
 
 	const bool bSuccess = Coordinator.ProcessConversionQueue(ConversionDelegate);
+	if (bSuccess && !CommandParams.bPreviewRequested)
+	{
+		BlueprintNativeCodeGenUtilsImpl::GenerateModuleBuildFile(CommandParams.ModuleName, Manifest);
+	}
+
 	// save the manifest regardless of success (so we can get an idea of how successful it was)
 	Manifest.Save();
 	return bSuccess & !ScopedErrorTracker.HasErrors();
diff --git a/Engine/Source/Developer/BlueprintNativeCodeGen/Private/NativeCodeGenCommandlineParams.cpp b/Engine/Source/Developer/BlueprintNativeCodeGen/Private/NativeCodeGenCommandlineParams.cpp
--- a/Engine/Source/Developer/BlueprintNativeCodeGen/Private/NativeCodeGenCommandlineParams.cpp
+++ b/Engine/Source/Developer/BlueprintNativeCodeGen/Private/NativeCodeGenCommandlineParams.cpp
@@ -123,6 +123,10 @@ FNativeCodeGenCommandlineParams::FNativeCodeGenCommandlineParams(const TArray<FS
 				UE_LOG(LogNativeCodeGenCommandline, Warning, TEXT("'%s' doesn't appear to be a valid output directory, defaulting to: '%s'"), *Value, *DefaultModulePath);
 			}
 		}
+		else if (!Switch.Compare(TEXT("moduleName"), ESearchCase::IgnoreCase))
+		{
+			ModuleName = Value;
+		}
 		else if (!Switch.Compare(TEXT("wipe"), ESearchCase::IgnoreCase))
 		{
 			bWipeRequested = true;
@@ -171,9 +175,16 @@ FNativeCodeGenCommandlineParams::FNativeCodeGenCommandlineParams(const TArray<FS
 
 	bool const bUtilizeExistingManifest = !ManifestFilePath.IsEmpty() && !bWipeRequested && FileManager.FileExists(*ManifestFilePath);
 	// an existing manifest would specify where to put the module
-	if (ModuleOutputDir.IsEmpty() && !bUtilizeExistingManifest)
+	if (!bUtilizeExistingManifest)
 	{
-		ModuleOutputDir = DefaultModulePath;
-	}
+		if (ModuleName.IsEmpty())
+		{
+			ModuleName = NativeCodeGenCommandlineParamsImpl::DefaultModuleName;
+		}
+		if (ModuleOutputDir.IsEmpty())
+		{
+			ModuleOutputDir = DefaultModulePath;
+		}
+	}	
 }
 
diff --git a/Engine/Source/Runtime/CoreUObject/Private/UObject/LinkerLoad.cpp b/Engine/Source/Runtime/CoreUObject/Private/UObject/LinkerLoad.cpp
--- a/Engine/Source/Runtime/CoreUObject/Private/UObject/LinkerLoad.cpp
+++ b/Engine/Source/Runtime/CoreUObject/Private/UObject/LinkerLoad.cpp
@@ -405,7 +405,7 @@ FLinkerLoad* FLinkerLoad::CreateLinker( UPackage* Parent, const TCHAR* Filename,
 /**
  * Looks for an existing linker for the given package, without trying to make one if it doesn't exist
  */
-FLinkerLoad* FLinkerLoad::FindExistingLinkerForPackage(UPackage* Package)
+FLinkerLoad* FLinkerLoad::FindExistingLinkerForPackage(const UPackage* Package)
 {
 	FLinkerLoad* Linker = nullptr;
 	if (Package)
diff --git a/Engine/Source/Runtime/CoreUObject/Public/UObject/Linker.h b/Engine/Source/Runtime/CoreUObject/Public/UObject/Linker.h
--- a/Engine/Source/Runtime/CoreUObject/Public/UObject/Linker.h
+++ b/Engine/Source/Runtime/CoreUObject/Public/UObject/Linker.h
@@ -1662,7 +1662,7 @@ public:
 	/**
 	 * Looks for an existing linker for the given package, without trying to make one if it doesn't exist
 	 */
-	COREUOBJECT_API static FLinkerLoad* FindExistingLinkerForPackage(UPackage* Package);
+	COREUOBJECT_API static FLinkerLoad* FindExistingLinkerForPackage(const UPackage* Package);
 
 	/**
 	 * Replaces OldObject's entry in its linker with NewObject, so that all subsequent loads of OldObject will return NewObject.
