CIS fix: Changing local variable name so that it doesn't shadow an input param.

diff --git a/Engine/Source/Developer/BlueprintNativeCodeGen/Private/BlueprintNativeCodeGenUtils.cpp b/Engine/Source/Developer/BlueprintNativeCodeGen/Private/BlueprintNativeCodeGenUtils.cpp
--- a/Engine/Source/Developer/BlueprintNativeCodeGen/Private/BlueprintNativeCodeGenUtils.cpp
+++ b/Engine/Source/Developer/BlueprintNativeCodeGen/Private/BlueprintNativeCodeGenUtils.cpp
@@ -71,14 +71,14 @@ static bool BlueprintNativeCodeGenUtilsImpl::GenerateModuleBuildFile(const FStri
 
 	for (UPackage* ModulePkg : ModulePackages)
 	{
-		const FString ModuleName = FPackageName::GetLongPackageAssetName(ModulePkg->GetName());
-		if (ModuleManager.ModuleExists(*ModuleName))
+		const FString PkgModuleName = FPackageName::GetLongPackageAssetName(ModulePkg->GetName());
+		if (ModuleManager.ModuleExists(*PkgModuleName))
 		{
-			PrivateDependencies.Add(ModuleName);
+			PrivateDependencies.Add(PkgModuleName);
 		}
 		else
 		{
-			UE_LOG(LogBlueprintCodeGen, Warning, TEXT("Failed to find module for package: %s"), *ModuleName);
+			UE_LOG(LogBlueprintCodeGen, Warning, TEXT("Failed to find module for package: %s"), *PkgModuleName);
 		}
 	}
 
