(WIP) Generate a full plugin for Blueprints being converted into C++ (instead of just a module).

diff --git a/Engine/Source/Developer/BlueprintNativeCodeGen/BlueprintNativeCodeGen.Build.cs b/Engine/Source/Developer/BlueprintNativeCodeGen/BlueprintNativeCodeGen.Build.cs
--- a/Engine/Source/Developer/BlueprintNativeCodeGen/BlueprintNativeCodeGen.Build.cs
+++ b/Engine/Source/Developer/BlueprintNativeCodeGen/BlueprintNativeCodeGen.Build.cs
@@ -28,7 +28,8 @@ public class BlueprintNativeCodeGen : ModuleRules
                 "Json",
                 "JsonUtilities",
                 "BlueprintCompilerCppBackend",
-                "GameProjectGeneration"
+                "GameProjectGeneration",
+                "Projects",
 			}
 		);
 
diff --git a/Engine/Source/Developer/BlueprintNativeCodeGen/Private/BlueprintNativeCodeGenManifest.cpp b/Engine/Source/Developer/BlueprintNativeCodeGen/Private/BlueprintNativeCodeGenManifest.cpp
--- a/Engine/Source/Developer/BlueprintNativeCodeGen/Private/BlueprintNativeCodeGenManifest.cpp
+++ b/Engine/Source/Developer/BlueprintNativeCodeGen/Private/BlueprintNativeCodeGenManifest.cpp
@@ -19,27 +19,89 @@ DEFINE_LOG_CATEGORY_STATIC(LogNativeCodeGenManifest, Log, All);
 namespace BlueprintNativeCodeGenManifestImpl
 {
 	static const int64 CPF_NoFlags = 0x00;
-	static const FString ManifestFileExt    = TEXT(".bpgen.manifest");
+	static const FString ManifestFileExt    = TEXT(".BpCodeGenManifest.json");
 	static const FString CppFileExt	        = TEXT(".cpp");
 	static const FString HeaderFileExt      = TEXT(".h");
 	static const FString HeaderSubDir       = TEXT("Public");
 	static const FString CppSubDir          = TEXT("Private");
 	static const FString ModuleBuildFileExt = TEXT(".Build.cs");
-	static const FString FallbackModuleName = TEXT("BpCodeGen");
-
-	/**  */
+	static const FString FallbackPluginName = TEXT("GeneratedBpCode");
+	static const FString PreviewFilePostfix = TEXT("-Preview");
+	static const FString PluginFileExt      = TEXT(".uplugin");
+	static const FString SourceSubDir       = TEXT("Source");
+
+	/**
+	 * Populates the provided manifest object with data from the specified file.
+	 * 
+	 * @param  FilePath    A json file path, denoting the file you want loaded and serialized in.
+	 * @param  Manifest    The target object that you want filled out with data from the file.
+	 * @return True if the manifest was successfully loaded, otherwise false.
+	 */
 	static bool LoadManifest(const FString& FilePath, FBlueprintNativeCodeGenManifest* Manifest);
 
-	static FString GetHeaderSaveDir(const FString& ModulePath);
-	static FString GetCppSaveDir(const FString& ModulePath);
-	static FString GetBaseFilename(const FAssetData& Asset);
-
-	/**  */
-	static FString GenerateHeaderSavePath(const FString& ModulePath, const FAssetData& Asset);
-	static FString GenerateCppSavePath(const FString& ModulePath, const FAssetData& Asset);
-
+	/**
+	 * Helper function the homogenizes file/directory paths so that they can be 
+	 * compared for equivalence against others.
+	 * 
+	 * @param  DirectoryPath    The path that you want sanitized.
+	 * @return A equivalent file/directory path, standardized for comparisons.
+	 */
 	static FString GetComparibleDirPath(const FString& DirectoryPath);
 
+	/**
+	 * Constructs a directory path for the plugin's source files.
+	 * 
+	 * @param  PluginPath    A path to the root directory for the plugin.
+	 * @return A path to the plugin's source directory.
+	 */
+	static FString GetTargetSourceDir(const FString& PluginPath);
+
+	/**
+	 * Constructs a path for the specified module, within the target plugin.
+	 * 
+	 * @param  PluginPath    A path to the root directory for the plugin.
+	 * @param  ModuleName    The name of the module that you're generating for the plugin.
+	 * @return A path to the module's directory, within the plugin.
+	 */
+	static FString GetModuleDirectory(const FString& PluginPath, const FString& ModuleName);
+
+	/**
+	 * Constructs a header path for the specified asset.
+	 * 
+	 * @param  PluginPath	A path to the root directory for the plugin.
+	 * @param  Asset		The asset you want a header file for.
+	 * @param  ModuleName	The name of the module that you're generating for the plugin.
+	 * @return A target file path for the specified asset to save a header to.
+	 */
+	static FString GenerateHeaderSavePath(const FString& PluginPath, const FAssetData& Asset, const FString& ModuleName);
+
+	/**
+	 * Constructs a cpp path for the specified asset.
+	 * 
+	 * @param  PluginPath	A path to the root directory for the plugin.
+	 * @param  Asset		The asset you want a cpp file for.
+	 * @param  ModuleName	The name of the module that you're generating for the plugin.
+	 * @return A target file path for the specified asset to save a cpp file to.
+	 */
+	static FString GenerateCppSavePath(const FString& PluginPath, const FAssetData& Asset, const FString& ModuleName);
+
+	/**
+	 * Coordinates with the code-gen backend, to produce a base filename (one 
+	 * without a file extension).
+	 * 
+	 * @param  Asset    The asset you want a filename for.
+	 * @return A filename (without extension) that matches the #include statements generated by the backend.
+	 */
+	static FString GetBaseFilename(const FAssetData& Asset);
+
+	/**
+	 * Collects native packages (which reflect distinct modules) that the 
+	 * specified object relies upon.
+	 * 
+	 * @param  AssetObj			The object you want dependencies for.
+	 * @param  DependenciesOut	An output array, which will be filled out with module packages that the target object relies upon.
+	 * @return False if the function failed to collect dependencies for the specified object.
+	 */
 	static bool GatherModuleDependencies(const UObject* AssetObj, TArray<UPackage*>& DependenciesOut);
 }
 
@@ -62,33 +124,33 @@ static bool BlueprintNativeCodeGenManifestImpl::LoadManifest(const FString& File
 }
 
 //------------------------------------------------------------------------------
-static FString BlueprintNativeCodeGenManifestImpl::GetHeaderSaveDir(const FString& ModulePath)
+static FString BlueprintNativeCodeGenManifestImpl::GetTargetSourceDir(const FString& PluginPath)
 {
-	return FPaths::Combine(*ModulePath, *HeaderSubDir);
+	return FPaths::Combine(*PluginPath, *SourceSubDir);
 }
 
 //------------------------------------------------------------------------------
-static FString BlueprintNativeCodeGenManifestImpl::GetCppSaveDir(const FString& ModulePath)
+static FString BlueprintNativeCodeGenManifestImpl::GetModuleDirectory(const FString& PluginPath, const FString& ModuleName)
 {
-	return FPaths::Combine(*ModulePath, *CppSubDir);
+	return FPaths::Combine(*GetTargetSourceDir(PluginPath), *ModuleName);
 }
 
 //------------------------------------------------------------------------------
-static FString BlueprintNativeCodeGenManifestImpl::GetBaseFilename(const FAssetData& Asset)
+static FString BlueprintNativeCodeGenManifestImpl::GenerateHeaderSavePath(const FString& PluginPath, const FAssetData& Asset, const FString& ModuleName)
 {
-	return IBlueprintCompilerCppBackendModule::GetBaseFilename(Asset.GetAsset());
+	return FPaths::Combine( *FPaths::Combine(*GetModuleDirectory(PluginPath, ModuleName), *HeaderSubDir), *GetBaseFilename(Asset) ) + HeaderFileExt;
 }
 
 //------------------------------------------------------------------------------
-static FString BlueprintNativeCodeGenManifestImpl::GenerateHeaderSavePath(const FString& ModulePath, const FAssetData& Asset)
+static FString BlueprintNativeCodeGenManifestImpl::GenerateCppSavePath(const FString& PluginPath, const FAssetData& Asset, const FString& ModuleName)
 {
-	return FPaths::Combine(*GetHeaderSaveDir(ModulePath), *GetBaseFilename(Asset)) + HeaderFileExt;
+	return FPaths::Combine( *FPaths::Combine(*GetModuleDirectory(PluginPath, ModuleName), *CppSubDir), *GetBaseFilename(Asset) ) + CppFileExt;
 }
 
 //------------------------------------------------------------------------------
-static FString BlueprintNativeCodeGenManifestImpl::GenerateCppSavePath(const FString& ModulePath, const FAssetData& Asset)
+static FString BlueprintNativeCodeGenManifestImpl::GetBaseFilename(const FAssetData& Asset)
 {
-	return FPaths::Combine(*GetCppSaveDir(ModulePath), *GetBaseFilename(Asset)) + CppFileExt;
+	return IBlueprintCompilerCppBackendModule::GetBaseFilename(Asset.GetAsset());
 }
 
 //------------------------------------------------------------------------------
@@ -153,12 +215,12 @@ static bool BlueprintNativeCodeGenManifestImpl::GatherModuleDependencies(const U
  ******************************************************************************/
  
 //------------------------------------------------------------------------------
-FConvertedAssetRecord::FConvertedAssetRecord(const FAssetData& AssetInfo, const FString& TargetModulePath)
+FConvertedAssetRecord::FConvertedAssetRecord(const FAssetData& AssetInfo, const FString& TargetDir, const FString& ModuleName)
 	: AssetPtr(AssetInfo.GetAsset())
 	, AssetType(AssetInfo.GetClass())
 	, AssetPath(AssetInfo.PackageName.ToString())
-	, GeneratedHeaderPath(BlueprintNativeCodeGenManifestImpl::GenerateHeaderSavePath(TargetModulePath, AssetInfo))
-	, GeneratedCppPath(BlueprintNativeCodeGenManifestImpl::GenerateCppSavePath(TargetModulePath, AssetInfo))
+	, GeneratedHeaderPath(BlueprintNativeCodeGenManifestImpl::GenerateHeaderSavePath(TargetDir, AssetInfo, ModuleName))
+	, GeneratedCppPath(BlueprintNativeCodeGenManifestImpl::GenerateCppSavePath(TargetDir, AssetInfo, ModuleName))
 {
 }
 
@@ -179,12 +241,13 @@ FString FBlueprintNativeCodeGenManifest::GetDefaultFilename()
 	return FApp::GetGameName() + BlueprintNativeCodeGenManifestImpl::ManifestFileExt;
 }
 
+//------------------------------------------------------------------------------
 FBlueprintNativeCodeGenManifest::FBlueprintNativeCodeGenManifest(const FString TargetPath)
-	: ModulePath(TargetPath)
+	: OutputDir(TargetPath)
 {
-	if (ModulePath.IsEmpty())
+	if (OutputDir.IsEmpty())
 	{
-		ModulePath = FPaths::GameIntermediateDir();
+		OutputDir = FPaths::GameIntermediateDir();
 	}
 	// do NOT load from an existing interface, as this is the default 
 	// constructor used by the USTRUCT() system
@@ -194,97 +257,131 @@ FBlueprintNativeCodeGenManifest::FBlueprintNativeCodeGenManifest(const FString T
 FBlueprintNativeCodeGenManifest::FBlueprintNativeCodeGenManifest(const FNativeCodeGenCommandlineParams& CommandlineParams)
 {
 	using namespace BlueprintNativeCodeGenManifestImpl;
-	bool const bLoadExistingManifest = !CommandlineParams.bWipeRequested || CommandlineParams.ModuleOutputDir.IsEmpty();
+	bool const bLoadExistingManifest = !CommandlineParams.bWipeRequested || CommandlineParams.OutputDir.IsEmpty();
 	
-	ModuleName = CommandlineParams.ModuleName;
-	ModulePath = CommandlineParams.ModuleOutputDir;
-	if (FPaths::IsRelative(ModulePath))
+	PluginName = CommandlineParams.PluginName;
+	OutputDir  = CommandlineParams.OutputDir;
+	if (FPaths::IsRelative(OutputDir))
 	{
-		FPaths::MakePathRelativeTo(ModulePath, *FPaths::GameDir());
+		FPaths::MakePathRelativeTo(OutputDir, *FPaths::GameDir());
 	}
 
 	if (!CommandlineParams.ManifestFilePath.IsEmpty())
 	{
-		ManifestPath = CommandlineParams.ManifestFilePath;
+		ManifestFilePath = CommandlineParams.ManifestFilePath;
 	}
 	else
 	{
-		if (!ensure(!ModuleName.IsEmpty()))
+		if (!ensure(!PluginName.IsEmpty()))
 		{
-			ModuleName = FallbackModuleName;
+			PluginName = FallbackPluginName;
 		}
-		if (!ensure(!ModulePath.IsEmpty()))
+		if (!ensure(!OutputDir.IsEmpty()))
 		{
-			ModulePath = FPaths::Combine(*FPaths::GameIntermediateDir(), *ModuleName);
+			OutputDir = FPaths::Combine(*FPaths::GameIntermediateDir(), *PluginName);
 		}
-		ManifestPath = FPaths::Combine(*GetModuleDir(), *GetDefaultFilename());
+
+		FString ManifestFilename = GetDefaultFilename();
+		if (CommandlineParams.bPreviewRequested)
+		{
+			ManifestFilename.InsertAt(ManifestFilename.Find(ManifestFileExt, ESearchCase::CaseSensitive, ESearchDir::FromEnd), PreviewFilePostfix);
+		}
+		ManifestFilePath = FPaths::Combine(*GetTargetDir(), *ManifestFilename);
 	}
 	
 	// incorporate an existing manifest (in case we're only re-converting a 
 	// handful of assets and adding them into an existing module)
-	if (bLoadExistingManifest && LoadManifest(ManifestPath, this))
+	if (bLoadExistingManifest && LoadManifest(ManifestFilePath, this))
 	{
-		// if they specified a separate module path, lets make sure we use that
+		// if they specified a separate plugin path, lets make sure we use that
 		// over what was found in the existing manifest
-		if (!CommandlineParams.ModuleOutputDir.IsEmpty())
+		if (!CommandlineParams.OutputDir.IsEmpty())
 		{
 			if (CommandlineParams.bWipeRequested)
 			{
-				ModulePath = CommandlineParams.ModuleOutputDir;
+				OutputDir = CommandlineParams.OutputDir;
 			}
 			else
 			{
-				const FString ExpectedPath = GetComparibleDirPath(CommandlineParams.ModuleOutputDir);
+				const FString ExpectedPath = GetComparibleDirPath(CommandlineParams.OutputDir);
 
-				FString TargetPath = GetModuleDir();
+				FString TargetPath = GetTargetDir();
 				TargetPath = GetComparibleDirPath(TargetPath);
 
 				if ( !FPaths::IsSamePath(ExpectedPath, TargetPath) )
 				{
 					UE_LOG(LogNativeCodeGenManifest, Error, 
-						TEXT("The existing manifest's module path does not match what was specified via the commandline."));
+						TEXT("The existing manifest's plugin path does not match what was specified via the commandline."));
 				}
 			}
 		}
 
-		const bool bNewModuleNameRequested = !CommandlineParams.ModuleName.IsEmpty() && (CommandlineParams.ModuleName != ModuleName);
-		if (bNewModuleNameRequested && !CommandlineParams.bPreviewRequested) 
+		const bool bNewPluginNameRequested = !CommandlineParams.PluginName.IsEmpty() && (CommandlineParams.PluginName != PluginName);
+		if (bNewPluginNameRequested && !CommandlineParams.bPreviewRequested)
 		{
-			// delete the old module file (if one exists)
-			IFileManager::Get().Delete(*GetBuildFilePath());
+			// delete the old plugin file (if one exists)
+			IFileManager::Get().Delete(*GetPluginFilePath());
 		}
 
-		// if we were only interested in obtaining the module path
+		// if we were only interested in obtaining the plugin path
 		if (CommandlineParams.bWipeRequested)
 		{
-			ModuleName = CommandlineParams.ModuleName;
+			PluginName = CommandlineParams.PluginName;
 			Clear();
 		}
 		else
 		{
-			if (bNewModuleNameRequested)
+			if (bNewPluginNameRequested)
 			{
-				UE_LOG(LogNativeCodeGenManifest, Warning, TEXT("The specified module name (%s) doesn't match the existing one (%s). Overridding with the new name."),
-					*CommandlineParams.ModuleName, *ModuleName);
-				ModuleName = CommandlineParams.ModuleName;
+				UE_LOG(LogNativeCodeGenManifest, Warning, TEXT("The specified plugin name (%s) doesn't match the existing one (%s). Overridding with the new name."),
+					*CommandlineParams.PluginName, *PluginName);
+				PluginName = CommandlineParams.PluginName;
 			}
-			MapConvertedAssets();
+			ConstructRecordLookupTable();
 		}
 	}
 }
 
+//------------------------------------------------------------------------------
+TArray<FString> FBlueprintNativeCodeGenManifest::GetDestinationPaths() const
+{
+	const FString TargetPath = GetTargetDir();
+
+	TArray<FString> DestPaths;
+	DestPaths.Add(ManifestFilePath);
+	DestPaths.Add(BlueprintNativeCodeGenManifestImpl::GetTargetSourceDir(TargetPath));
+	DestPaths.Add(GetPluginFilePath());
+
+	return DestPaths;
+}
+
+//------------------------------------------------------------------------------
+FString FBlueprintNativeCodeGenManifest::GetPluginFilePath() const
+{
+	return FPaths::Combine(*GetTargetDir(), *PluginName) + BlueprintNativeCodeGenManifestImpl::PluginFileExt;
+}
+
+//------------------------------------------------------------------------------
+FString FBlueprintNativeCodeGenManifest::GetModuleFilePath() const
+{
+	return FPaths::Combine( *BlueprintNativeCodeGenManifestImpl::GetModuleDirectory(GetTargetDir(), GetTargetModuleName()), *GetTargetModuleName() ) + 
+		BlueprintNativeCodeGenManifestImpl::ModuleBuildFileExt;
+}
+
 //------------------------------------------------------------------------------
 FConvertedAssetRecord& FBlueprintNativeCodeGenManifest::CreateConversionRecord(const FAssetData& AssetInfo)
 {
-	const FString AssetPath = AssetInfo.PackageName.ToString();
-	const FString TargetModulePath = GetModuleDir();
+	const FString AssetPath  = AssetInfo.PackageName.ToString();
+	const FString TargetPath = GetTargetDir();
 	// load the asset (if it isn't already)
-	const UObject* AssetObj = AssetInfo.GetAsset();
+	const UObject* AssetObj   = AssetInfo.GetAsset();
+	const FString& ModuleName = GetTargetModuleName();
 
 	FConvertedAssetRecord* ConversionRecordPtr = FindConversionRecord(AssetPath);
 	if (ConversionRecordPtr == nullptr)
 	{
-		ConversionRecordPtr = &ConvertedAssets[ ConvertedAssets.Add(FConvertedAssetRecord(AssetInfo, TargetModulePath)) ];
+		RecordLookupTable.Add(AssetPath, ConvertedAssets.Num());
+		ConversionRecordPtr = &ConvertedAssets[ ConvertedAssets.Add(FConvertedAssetRecord(AssetInfo, TargetPath, ModuleName)) ];
 	}
 	else if ( !ensure((ConversionRecordPtr->AssetPath == AssetPath) && (ConversionRecordPtr->AssetType == AssetInfo.GetClass())) )
 	{
@@ -307,8 +404,8 @@ FConvertedAssetRecord& FBlueprintNativeCodeGenManifest::CreateConversionRecord(c
 		}
 		if (ConversionRecord.GeneratedHeaderPath.IsEmpty())
 		{
-			ConversionRecord.GeneratedHeaderPath = BlueprintNativeCodeGenManifestImpl::GenerateHeaderSavePath(TargetModulePath, AssetInfo);
-			ConversionRecord.GeneratedCppPath    = BlueprintNativeCodeGenManifestImpl::GenerateCppSavePath(TargetModulePath, AssetInfo);
+			ConversionRecord.GeneratedHeaderPath = BlueprintNativeCodeGenManifestImpl::GenerateHeaderSavePath(TargetPath, AssetInfo, ModuleName);
+			ConversionRecord.GeneratedCppPath    = BlueprintNativeCodeGenManifestImpl::GenerateCppSavePath(TargetPath, AssetInfo, ModuleName);
 		}
 		ensure(ConversionRecord.IsValid());
 	}
@@ -320,7 +417,7 @@ FConvertedAssetRecord& FBlueprintNativeCodeGenManifest::CreateConversionRecord(c
 FConvertedAssetRecord* FBlueprintNativeCodeGenManifest::FindConversionRecord(const FString& AssetPath, bool bSlowSearch)
 {
 	FConvertedAssetRecord* FoundRecord = nullptr;
-	if (int32* IndexPtr = RecordMap.Find(AssetPath))
+	if (int32* IndexPtr = RecordLookupTable.Find(AssetPath))
 	{
 		check(*IndexPtr < ConvertedAssets.Num());
 		FoundRecord = &ConvertedAssets[*IndexPtr];
@@ -336,6 +433,10 @@ FConvertedAssetRecord* FBlueprintNativeCodeGenManifest::FindConversionRecord(con
 			}
 		}
 	}
+	else
+	{
+		check(RecordLookupTable.Num() == ConvertedAssets.Num());
+	}
 	return FoundRecord;
 }
 
@@ -348,26 +449,6 @@ bool FBlueprintNativeCodeGenManifest::LogDependencies(const FAssetData& AssetInf
 	return BlueprintNativeCodeGenManifestImpl::GatherModuleDependencies(AssetObj, ModuleDependencies);
 }
 
-//------------------------------------------------------------------------------
-TArray<FString> FBlueprintNativeCodeGenManifest::GetTargetPaths() const 
-{
-	const FString TargetModulePath = GetModuleDir();
-
-	TArray<FString> DestPaths;
-	DestPaths.Add(ManifestPath);
-	DestPaths.Add(BlueprintNativeCodeGenManifestImpl::GetHeaderSaveDir(TargetModulePath));
-	DestPaths.Add(BlueprintNativeCodeGenManifestImpl::GetCppSaveDir(TargetModulePath));
-	DestPaths.Add(GetBuildFilePath());
-
-	return DestPaths;
-}
-
-//------------------------------------------------------------------------------
-FString FBlueprintNativeCodeGenManifest::GetBuildFilePath() const
-{
-	return FPaths::Combine(*GetModuleDir(), *ModuleName) + BlueprintNativeCodeGenManifestImpl::ModuleBuildFileExt;
-}
-
 //------------------------------------------------------------------------------
 bool FBlueprintNativeCodeGenManifest::Save(bool bSort) const
 {
@@ -381,6 +462,7 @@ bool FBlueprintNativeCodeGenManifest::Save(bool bSort) const
 			}
 			return Lhs.AssetType->GetName() < Rhs.AssetType->GetName();
 		});
+		RecordLookupTable.Empty();
 	}
 
 	TSharedRef<FJsonObject> JsonObject = MakeShareable(new FJsonObject());
@@ -394,40 +476,40 @@ bool FBlueprintNativeCodeGenManifest::Save(bool bSort) const
 		if (FJsonSerializer::Serialize(JsonObject, JsonWriter))
 		{
 			JsonWriter->Close();
-			return FFileHelper::SaveStringToFile(FileContents, *ManifestPath);
+			return FFileHelper::SaveStringToFile(FileContents, *ManifestFilePath);
 		}
 	}
 	return false;
 }
 
 //------------------------------------------------------------------------------
-void FBlueprintNativeCodeGenManifest::Clear()
-{
-	ConvertedAssets.Empty();
-	RecordMap.Empty();
-}
-
-//------------------------------------------------------------------------------
-FString FBlueprintNativeCodeGenManifest::GetModuleDir() const
+FString FBlueprintNativeCodeGenManifest::GetTargetDir() const
 {
-	FString TargetPath = ModulePath;
-	if (FPaths::IsRelative(ModulePath))
+	FString TargetPath = OutputDir;
+	if (FPaths::IsRelative(TargetPath))
 	{
-		TargetPath = FPaths::ConvertRelativePathToFull(FPaths::GameDir(), ModulePath);
+		TargetPath = FPaths::ConvertRelativePathToFull(FPaths::GameDir(), TargetPath);
 		TargetPath = FPaths::ConvertRelativePathToFull(TargetPath);
 	}
 	return TargetPath;
 }
 
+
 //------------------------------------------------------------------------------
-void FBlueprintNativeCodeGenManifest::MapConvertedAssets()
+void FBlueprintNativeCodeGenManifest::Clear()
 {
-	RecordMap.Empty(ConvertedAssets.Num());
+	ConvertedAssets.Empty();
+	RecordLookupTable.Empty();
+}
+
+//------------------------------------------------------------------------------
+void FBlueprintNativeCodeGenManifest::ConstructRecordLookupTable()
+{
+	RecordLookupTable.Empty(ConvertedAssets.Num());
 
 	for (int32 RecordIndex = 0; RecordIndex < ConvertedAssets.Num(); ++RecordIndex)
 	{
 		const FConvertedAssetRecord& AssetRecord = ConvertedAssets[RecordIndex];
-		RecordMap.Add(AssetRecord.AssetPath, RecordIndex);
+		RecordLookupTable.Add(AssetRecord.AssetPath, RecordIndex);
 	}
 }
-
diff --git a/Engine/Source/Developer/BlueprintNativeCodeGen/Private/BlueprintNativeCodeGenManifest.h b/Engine/Source/Developer/BlueprintNativeCodeGen/Private/BlueprintNativeCodeGenManifest.h
--- a/Engine/Source/Developer/BlueprintNativeCodeGen/Private/BlueprintNativeCodeGenManifest.h
+++ b/Engine/Source/Developer/BlueprintNativeCodeGen/Private/BlueprintNativeCodeGenManifest.h
@@ -19,8 +19,11 @@ struct FConvertedAssetRecord
 
 public:
 	FConvertedAssetRecord() {}
-	FConvertedAssetRecord(const FAssetData& AssetInfo, const FString& TargetModulePath);
+	FConvertedAssetRecord(const FAssetData& AssetInfo, const FString& TargetDir, const FString& ModuleName);
 
+	/**
+	 * @return True if this record demonstates a successful conversion; false if it is incomplete 
+	 */
 	bool IsValid();
 
 public:
@@ -48,6 +51,7 @@ struct FBlueprintNativeCodeGenManifest
 {
 	GENERATED_USTRUCT_BODY()
 
+	/** @return The default filename for manifest files. */
 	static FString GetDefaultFilename();
 
 public: 
@@ -55,69 +59,103 @@ public:
 	FBlueprintNativeCodeGenManifest(const FNativeCodeGenCommandlineParams& CommandlineParams);
 
 	/**
-	 * A query that retrieves a list of file/directory paths. These paths are 
-	 * the destination targets for files and directories that should/will be 
-	 * generated as part of the conversion process.
+	 * A query that retrieves a list of file/directory paths. Conceptually, 
+	 * these files/directories are fully owned by the the conversion process.
 	 * 
-	 * @return An array of file/directory paths that serve as destinations for the conversion process.
+	 * @return An array of file/directory paths that will be generated the conversion process.
 	 */
-	TArray<FString> GetTargetPaths() const;
+	TArray<FString> GetDestinationPaths() const;
 
-	/**  */
-	FString GetBuildFilePath() const;
+	/**
+	 * Constructs a .uplugin path, specifying the destination for the plugin's
+	 * descriptor file.
+	 * 
+	 * @return The destination filepath for the .uplugin file.
+	 */
+	FString GetPluginFilePath() const;
 
 	/**
+	 * Constructs a .Build.cs path, detailing the destination for the module's 
+	 * build file.
 	 * 
+	 * @return The destination filepath for the .Build.cs file.
+	 */
+	FString GetModuleFilePath() const;
+
+	/**
+	 * Logs an entry detailing the specified asset's conversion (the asset's 
+	 * name, the resulting cpp/h files, etc.). Will return an existing entry, 
+	 * if one exists for the specified asset.
 	 * 
-	 * @param  AssetInfo    
-	 * @return 
+	 * @param  AssetInfo    Defines the asset that you plan on converting (and want to make an entry for).
+	 * @return A record detailing the asset's expected conversion.
 	 */
 	FConvertedAssetRecord& CreateConversionRecord(const FAssetData& AssetInfo);
 
 	/**
+	 * Searches the manifest's conversion records for an entry pertaining to the
+	 * specified asset.
 	 * 
-	 * 
-	 * @param  AssetPath    
-	 * @return 
+	 * @param  AssetPath    Serves as a lookup key for the asset in question.
+	 * @return A pointer to the found conversion record, null if one doesn't exist.
 	 */
 	FConvertedAssetRecord* FindConversionRecord(const FString& AssetPath, bool bSlowSearch = false);
 
-	/**  */
+	/**
+	 * Records module dependencies for the specified asset.
+	 * 
+	 * @param  AssetInfo    The asset you want to collect dependencies for.
+	 * @return True if we were able to successfully collect dependencies from the asset; false if the process failed.
+	 */
 	bool LogDependencies(const FAssetData& AssetInfo);
 
-	/**  */
+	/** 
+	 * @return A list of all known modules that this plugin will depend on. 
+	 */
 	const TArray<UPackage*>& GetModuleDependencies() const { return ModuleDependencies; }
 
 	/**
+	 * Saves this manifest as json, to its target destination (which it was 
+	 * setup with).
 	 * 
+	 * @param  bSort    Defines if you want the manifest's conversion log sorted before the files is saved (making it easier to read).
+	 * @return True if the save was a success, otherwise false.
 	 */
 	bool Save(bool bSort = true) const;
 
 private:
 	/**
-	 * 
+	 * @return The destination directory for the plugin and all its related files.
 	 */
-	FString GetModuleDir() const;
+	FString GetTargetDir() const;
 
 	/**
-	 * 
+	 * @return The target name for the plugin's primary module.
+	 */
+	const FString& GetTargetModuleName() const { return PluginName; }
+
+	/**
+	 * Empties the manifest, leaving only the destination directory and file
+	 * names intact.
 	 */
 	void Clear();
 
 	/**
-	 *
+	 * Maps the asset log so that we can easily query for specific conversion
+	 * records.
 	 */
-	void MapConvertedAssets();
+	void ConstructRecordLookupTable();
 	
-	/**  */
-	FString ManifestPath;
+private:
+	/** Defines the destination filepath for this manifest */
+	FString ManifestFilePath;
 
 	UPROPERTY()
-	FString ModuleName;
+	FString PluginName;
 
 	/** Relative to the project's directory */
 	UPROPERTY()
-	FString ModulePath;
+	FString OutputDir;
 
 	UPROPERTY()
 	TArray<UPackage*> ModuleDependencies;
@@ -126,6 +164,6 @@ private:
 	UPROPERTY()
 	mutable TArray<FConvertedAssetRecord> ConvertedAssets;
 
-	/**  */
-	TMap<FString, int32> RecordMap;
+	/** Serves as a quick way to look up records in the ConvertedAssets array */
+	mutable TMap<FString, int32> RecordLookupTable;
 };
diff --git a/Engine/Source/Developer/BlueprintNativeCodeGen/Private/BlueprintNativeCodeGenUtils.cpp b/Engine/Source/Developer/BlueprintNativeCodeGen/Private/BlueprintNativeCodeGenUtils.cpp
--- a/Engine/Source/Developer/BlueprintNativeCodeGen/Private/BlueprintNativeCodeGenUtils.cpp
+++ b/Engine/Source/Developer/BlueprintNativeCodeGen/Private/BlueprintNativeCodeGenUtils.cpp
@@ -24,8 +24,30 @@ DEFINE_LOG_CATEGORY(LogBlueprintCodeGen)
 
 namespace BlueprintNativeCodeGenUtilsImpl
 {
+	/**
+	 * Deletes the files/directories in the supplied array.
+	 * 
+	 * @param  TargetPaths    The set of directory and file paths that you want deleted.
+	 * @return True if all the files/directories were successfully deleted, other wise false.
+	 */
 	static bool WipeTargetPaths(const TArray<FString>& TargetPaths);
-	static bool GenerateModuleBuildFile(const FString& ModuleName, const FBlueprintNativeCodeGenManifest& Manifest);
+	
+	/**
+	 * Creates and fills out a new .uplugin file for the converted assets.
+	 * 
+	 * @param  PluginName	The name of the plugin you're generating.
+	 * @param  Manifest		Defines where the plugin file should be saved.
+	 * @return True if the file was successfully saved, otherwise false.
+	 */
+	static bool GeneratePluginDescFile(const FString& PluginName, const FBlueprintNativeCodeGenManifest& Manifest);
+
+	/**
+	 * Creates and fills out a new .Build.cs file for the target module.
+	 * 
+	 * @param  Manifest    Defines where the module file should be saved, what it should be named, etc..
+	 * @return True if the file was successfully saved, otherwise false.
+	 */
+	static bool GenerateModuleBuildFile(const FBlueprintNativeCodeGenManifest& Manifest);
 }
 
 //------------------------------------------------------------------------------
@@ -50,7 +72,40 @@ static bool BlueprintNativeCodeGenUtilsImpl::WipeTargetPaths(const TArray<FStrin
 }
 
 //------------------------------------------------------------------------------
-static bool BlueprintNativeCodeGenUtilsImpl::GenerateModuleBuildFile(const FString& ModuleName, const FBlueprintNativeCodeGenManifest& Manifest)
+static bool BlueprintNativeCodeGenUtilsImpl::GeneratePluginDescFile(const FString& PluginName, const FBlueprintNativeCodeGenManifest& Manifest)
+{
+	FPluginDescriptor PluginDesc;
+	PluginDesc.FriendlyName = PluginName;
+	PluginDesc.CreatedBy    = TEXT("Epic Games, Inc.");
+	PluginDesc.CreatedByURL = TEXT("http://epicgames.com");
+	PluginDesc.Description  = TEXT("A programatically generated plugin which contains source files produced from Blueprint assets. The aim of this is to help performance by eliminating script overhead for the converted assets (using the source files in place of thier coresponding assets).");
+	PluginDesc.DocsURL      = TEXT("@TODO");
+	PluginDesc.SupportURL   = TEXT("https://answers.unrealengine.com/");
+	PluginDesc.Category     = TEXT("Intermediate");
+	PluginDesc.bEnabledByDefault  = true;
+	PluginDesc.bCanContainContent = false;
+	PluginDesc.bIsBetaVersion     = true; // @TODO: change once we're confident in the feature
+
+	FModuleDescriptor ModuleDesc;
+	ModuleDesc.Name = *PluginName;
+	ModuleDesc.Type = EHostType::Runtime;
+	// load at startup (during engine init), after game modules have been loaded 
+	ModuleDesc.LoadingPhase = ELoadingPhase::Default; 
+
+	PluginDesc.Modules.Add(ModuleDesc);
+
+	FText ErrorMessage;
+	bool bSuccess = PluginDesc.Save(Manifest.GetPluginFilePath(), ErrorMessage);
+
+	if (!bSuccess)
+	{
+		UE_LOG(LogBlueprintCodeGen, Error, TEXT("Failed to generate the plugin description file: %s"), *ErrorMessage.ToString());
+	}
+	return bSuccess;
+}
+
+//------------------------------------------------------------------------------
+static bool BlueprintNativeCodeGenUtilsImpl::GenerateModuleBuildFile(const FBlueprintNativeCodeGenManifest& Manifest)
 {
 	FModuleManager& ModuleManager = FModuleManager::Get();
 	
@@ -82,15 +137,21 @@ static bool BlueprintNativeCodeGenUtilsImpl::GenerateModuleBuildFile(const FStri
 		}
 	}
 
-	const FString BuildFilePath = Manifest.GetBuildFilePath();
+	const FString BuildFilePath = Manifest.GetModuleFilePath();
+
+	FString ModuleName = FPaths::GetCleanFilename(BuildFilePath);
+	int32 ExtIndex = INDEX_NONE;
+	ModuleName.FindChar('.', ExtIndex);
+	ModuleName = ModuleName.Left(ExtIndex);
+
 	FText ErrorMessage;
-	GameProjectUtils::GenerateGameModuleBuildFile(BuildFilePath, ModuleName, PublicDependencies, PrivateDependencies, ErrorMessage);
+	bool bSuccess = GameProjectUtils::GenerateGameModuleBuildFile(BuildFilePath, ModuleName, PublicDependencies, PrivateDependencies, ErrorMessage);
 
-	if (!ErrorMessage.IsEmpty())
+	if (!bSuccess)
 	{
 		UE_LOG(LogBlueprintCodeGen, Error, TEXT("Failed to generate module build file: %s"), *ErrorMessage.ToString());
 	}
-	return !ErrorMessage.IsEmpty();
+	return bSuccess;
 }
 
 /*******************************************************************************
@@ -98,7 +159,7 @@ static bool BlueprintNativeCodeGenUtilsImpl::GenerateModuleBuildFile(const FStri
  ******************************************************************************/
 
 //------------------------------------------------------------------------------
-bool FBlueprintNativeCodeGenUtils::GenerateCodeModule(const FNativeCodeGenCommandlineParams& CommandParams)
+bool FBlueprintNativeCodeGenUtils::GeneratePlugin(const FNativeCodeGenCommandlineParams& CommandParams)
 {
 	FScopedFeedbackContext ScopedErrorTracker;
 	FBlueprintNativeCodeGenCoordinator Coordinator(CommandParams);
@@ -112,7 +173,7 @@ bool FBlueprintNativeCodeGenUtils::GenerateCodeModule(const FNativeCodeGenComman
 
 	if (CommandParams.bWipeRequested && !CommandParams.bPreviewRequested)
 	{
-		TArray<FString> TargetPaths = Manifest.GetTargetPaths();
+		TArray<FString> TargetPaths = Manifest.GetDestinationPaths();
 		if (!BlueprintNativeCodeGenUtilsImpl::WipeTargetPaths(TargetPaths))
 		{
 			UE_LOG(LogBlueprintCodeGen, Warning, TEXT("Failed to wipe target files/directories."));
@@ -171,13 +232,18 @@ bool FBlueprintNativeCodeGenUtils::GenerateCodeModule(const FNativeCodeGenComman
 		ConversionDelegate.BindLambda([](FConvertedAssetRecord& ConversionRecord){ return true; });
 	}
 
-	const bool bSuccess = Coordinator.ProcessConversionQueue(ConversionDelegate);
+	bool bSuccess = Coordinator.ProcessConversionQueue(ConversionDelegate);
+
 	if (bSuccess && !CommandParams.bPreviewRequested)
 	{
-		BlueprintNativeCodeGenUtilsImpl::GenerateModuleBuildFile(CommandParams.ModuleName, Manifest);
+		bSuccess &= BlueprintNativeCodeGenUtilsImpl::GenerateModuleBuildFile(Manifest);
+		if (bSuccess)
+		{
+			bSuccess &= BlueprintNativeCodeGenUtilsImpl::GeneratePluginDescFile(CommandParams.PluginName, Manifest);
+		}
 	}
 
-	// save the manifest regardless of success (so we can get an idea of how successful it was)
+	// save the manifest regardless of success (so we can get an idea of where failures broke down)
 	Manifest.Save();
 	return bSuccess & !ScopedErrorTracker.HasErrors();
 }
diff --git a/Engine/Source/Developer/BlueprintNativeCodeGen/Private/BlueprintNativeCodeGenUtils.h b/Engine/Source/Developer/BlueprintNativeCodeGen/Private/BlueprintNativeCodeGenUtils.h
--- a/Engine/Source/Developer/BlueprintNativeCodeGen/Private/BlueprintNativeCodeGenUtils.h
+++ b/Engine/Source/Developer/BlueprintNativeCodeGen/Private/BlueprintNativeCodeGenUtils.h
@@ -17,10 +17,10 @@ public:
 	/**
 	 * 
 	 * 
-	 * @param  Coordinator    
+	 * @param  CommandParams    
 	 * @return 
 	 */
-	static bool GenerateCodeModule(const FNativeCodeGenCommandlineParams& CommandParams);
+	static bool GeneratePlugin(const FNativeCodeGenCommandlineParams& CommandParams);
 
 	/**
 	 * Recompiles the bytecode of a blueprint only. Should only be run for
diff --git a/Engine/Source/Developer/BlueprintNativeCodeGen/Private/GenerateBlueprintCodeModuleCommandlet.cpp b/Engine/Source/Developer/BlueprintNativeCodeGen/Private/GenerateBlueprintCodeModuleCommandlet.cpp
deleted file mode 100644
diff --git a/Engine/Source/Developer/BlueprintNativeCodeGen/Private/GenerateBlueprintCodeModuleCommandlet.h b/Engine/Source/Developer/BlueprintNativeCodeGen/Private/GenerateBlueprintCodeModuleCommandlet.h
deleted file mode 100644
diff --git a/Engine/Source/Developer/BlueprintNativeCodeGen/Private/GenerateNativePluginFromBlueprintCommandlet.cpp b/Engine/Source/Developer/BlueprintNativeCodeGen/Private/GenerateNativePluginFromBlueprintCommandlet.cpp
new file mode 100644
--- /dev/null
+++ b/Engine/Source/Developer/BlueprintNativeCodeGen/Private/GenerateNativePluginFromBlueprintCommandlet.cpp
@@ -0,0 +1,40 @@
+// Copyright 1998-2015 Epic Games, Inc. All Rights Reserved.
+
+#include "BlueprintNativeCodeGenPCH.h"
+#include "GenerateNativePluginFromBlueprintCommandlet.h"
+#include "NativeCodeGenCommandlineParams.h"
+#include "BlueprintNativeCodeGenCoordinator.h"
+#include "BlueprintNativeCodeGenUtils.h"
+#include "BlueprintNativeCodeGenManifest.h"
+#include "FileManager.h"
+
+/*******************************************************************************
+ * UGenerateNativePluginFromBlueprintCommandlet
+*******************************************************************************/
+
+//------------------------------------------------------------------------------
+UGenerateNativePluginFromBlueprintCommandlet::UGenerateNativePluginFromBlueprintCommandlet(FObjectInitializer const& ObjectInitializer)
+	: Super(ObjectInitializer)
+{
+}
+
+//------------------------------------------------------------------------------
+int32 UGenerateNativePluginFromBlueprintCommandlet::Main(FString const& Params)
+{
+	TArray<FString> Tokens, Switches;
+	ParseCommandLine(*Params, Tokens, Switches);
+
+	FNativeCodeGenCommandlineParams CommandlineParams(Switches);
+
+	if (CommandlineParams.bHelpRequested)
+	{
+		UE_LOG(LogBlueprintCodeGen, Display, TEXT("%s"), *FNativeCodeGenCommandlineParams::HelpMessage);
+		return 0;
+	}
+
+	FBlueprintNativeCodeGenUtils::FScopedFeedbackContext ScopedErrorTracker;
+	const bool bSuccess = FBlueprintNativeCodeGenUtils::GeneratePlugin(CommandlineParams);
+
+	return bSuccess && ScopedErrorTracker.HasErrors();
+}
+
diff --git a/Engine/Source/Developer/BlueprintNativeCodeGen/Private/GenerateNativePluginFromBlueprintCommandlet.h b/Engine/Source/Developer/BlueprintNativeCodeGen/Private/GenerateNativePluginFromBlueprintCommandlet.h
new file mode 100644
--- /dev/null
+++ b/Engine/Source/Developer/BlueprintNativeCodeGen/Private/GenerateNativePluginFromBlueprintCommandlet.h
@@ -0,0 +1,18 @@
+// Copyright 1998-2015 Epic Games, Inc. All Rights Reserved.
+
+#pragma once
+
+#include "Commandlets/Commandlet.h"
+#include "GenerateNativePluginFromBlueprintCommandlet.generated.h"
+
+/**  */
+UCLASS()
+class UGenerateNativePluginFromBlueprintCommandlet : public UCommandlet
+{
+	GENERATED_UCLASS_BODY()
+
+public:		
+	// Begin UCommandlet Interface
+	virtual int32 Main(FString const& Params) override;
+	// End UCommandlet Interface
+};
diff --git a/Engine/Source/Developer/BlueprintNativeCodeGen/Private/NativeCodeGenCommandlineParams.cpp b/Engine/Source/Developer/BlueprintNativeCodeGen/Private/NativeCodeGenCommandlineParams.cpp
--- a/Engine/Source/Developer/BlueprintNativeCodeGen/Private/NativeCodeGenCommandlineParams.cpp
+++ b/Engine/Source/Developer/BlueprintNativeCodeGen/Private/NativeCodeGenCommandlineParams.cpp
@@ -1,5 +1,4 @@
 // Copyright 1998-2015 Epic Games, Inc. All Rights Reserved.
-// Copyright 1998-2015 Epic Games, Inc. All Rights Reserved.
 
 #include "BlueprintNativeCodeGenPCH.h"
 #include "NativeCodeGenCommandlineParams.h"
@@ -11,14 +10,14 @@ DEFINE_LOG_CATEGORY_STATIC(LogNativeCodeGenCommandline, Log, All);
 const FString FNativeCodeGenCommandlineParams::HelpMessage = TEXT("\n\
 \n\
 -------------------------------------------------------------------------------\n\
-::  GenerateBlueprintCodeModule    ::    Converts Blueprint assets into C++    \n\
+::  GenerateNativePluginFromBlueprint :: Converts Blueprint assets into C++    \n\
 -------------------------------------------------------------------------------\n\
 \n\
 ::                                                                             \n\
 :: Usage                                                                       \n\
 ::                                                                             \n\
 \n\
-    UE4Editor.exe <project> -run=GenerateBlueprintCodeModule [parameters]      \n\
+    UE4Editor.exe <project> -run=GenerateNativePluginFromBlueprint [parameters]\n\
 \n\
 ::                                                                             \n\
 :: Parameters                                                                  \n\
@@ -30,34 +29,38 @@ const FString FNativeCodeGenCommandlineParams::HelpMessage = TEXT("\n\
                         the form: /Game/MyContentDir,/Engine/EngineAssetName.  \n\
 \n\
     -blacklist=<Assets> Explicitly specifies assets that you don't want        \n\
-                        converted (listed in the same manner as -whitelist).   \n\
-                        This takes priority over the whitelist (even if it     \n\
-                        results in uncompilable code).                         \n\
+                        converted (listed in the same manner as the -whitelist \n\
+                        param). This takes priority over the whitelist (even if\n\
+                        it results in uncompilable code).                      \n\
 \n\
-    -output=<ModuleDir> Specifies the path where you want converted assets     \n\
+    -output=<TargetDir> Specifies the path where you want converted assets     \n\
                         saved to. If left unset, this will default to the      \n\
                         project's intermeadiate folder (within a sub-folder of \n\
                         its own). If specified as a relative path, it will be  \n\
-                        relative to the target project's root directory.       \n\
+                        interpreted as relative to the project directory.      \n\
 \n\
-    -moduleName=<Name>  Defines the output module's name. If left unset, then  \n\
-                        the module will default to... \n\
+    -pluginName=<Name>  Defines the generated plugin's name. If the output     \n\
+                        directory doesn't match this in name, then a sub-folder\n\
+                        will be created there to house the plugin.             \n\
 \n\
-    -wipe               Will clear out the target directory when specified.    \n\
-                        Normally, without this switch, if the target directory \n\
-                        contains source from a previous run, it builds atop the\n\
-                        existing module and appends to its existing manifest.  \n\
+    -noWipe             By default, we'll clear out the target directory from  \n\
+                        previous conversions. However, with this, if the       \n\
+						target directory contains source from a previous run,  \n\
+                        we'll build atop the existing module and appends to its\n\
+						existing manifest.                                     \n\
 \n\
-    -manifest=<Path>    Specifies where to save the conversion's manifest file.\n\
-                        If <Path> is not an existing directory or file, then it\n\
-                        is assumed that it is a file path. If a manifest       \n\
-                        already exists here, then we can use this in place of  \n\
-                        the -output command (and the target module directory   \n\
-                        will be determined by the old manifest file). If left  \n\
-                        unset, then it will default to the module's target     \n\
-                        directory.                                             \n\
+    -manifest=<Path>    Specifies where to save the resultant manifest file. If\n\
+                        <Path> is not an existing directory or file, then it is\n\
+                        assumed that this is a file path. If a manifest already\n\
+                        exists here, then it'll be used to define the target   \n\
+                        output directory (which will be read from the existing \n\
+                        manifest file). If left unset then the manifest will   \n\
+                        default to the plugin's directory.                     \n\
 \n\
-    -preview            \n\
+    -preview            When specified, the process will NOT wipe any existing \n\
+                        files, and the target assets will NOT be compiled and  \n\
+						converted. Instead, just a manifest file will be       \n\
+						generated, detailing the command's predicted output.   \n\
 \n\
     -help, -h, -?       Display this message and then exit.                    \n\
 \n");
@@ -68,10 +71,27 @@ const FString FNativeCodeGenCommandlineParams::HelpMessage = TEXT("\n\
 
 namespace NativeCodeGenCommandlineParamsImpl
 {
-	static const FString DefaultModuleName(TEXT("BpCodeGen"));
+	static const FString DefaultPluginName(TEXT("GeneratedBpCode"));
 	static const FString ParamListDelim(TEXT(","));
+
+	/**
+	 * Cleans up  the passed in path from the commandline (removes quotes, 
+	 * corrects slashes, etc.), so that it is ready to be utilized by the 
+	 * conversion process.
+	 * 
+	 * @param  PathInOut    The path you want standardized.
+	 */
+	static void SanitizePathParam(FString& PathInOut);
 }
  
+//------------------------------------------------------------------------------
+static void NativeCodeGenCommandlineParamsImpl::SanitizePathParam(FString& PathInOut)
+{
+	PathInOut.RemoveFromStart(TEXT("\""));
+	PathInOut.RemoveFromEnd(TEXT("\""));
+	FPaths::NormalizeDirectoryName(PathInOut);
+}
+
 /*******************************************************************************
  * FNativeCodeGenCommandlineParams
  ******************************************************************************/
@@ -79,11 +99,12 @@ namespace NativeCodeGenCommandlineParamsImpl
 //------------------------------------------------------------------------------
 FNativeCodeGenCommandlineParams::FNativeCodeGenCommandlineParams(const TArray<FString>& CommandlineSwitches)
 	: bHelpRequested(false)
-	, bWipeRequested(false)
+	, bWipeRequested(true)
 	, bPreviewRequested(false)
 {
+	using namespace NativeCodeGenCommandlineParamsImpl;
+
 	IFileManager& FileManager = IFileManager::Get();
-	const FString DefaultModulePath = FPaths::Combine(*FPaths::GameIntermediateDir(), *NativeCodeGenCommandlineParamsImpl::DefaultModuleName);
 
 	for (const FString& Param : CommandlineSwitches)
 	{
@@ -106,7 +127,7 @@ FNativeCodeGenCommandlineParams::FNativeCodeGenCommandlineParams(const TArray<FS
 		}
 		else if (!Switch.Compare(TEXT("output"), ESearchCase::IgnoreCase))
 		{
-			FPaths::NormalizeDirectoryName(Value);
+			SanitizePathParam(Value);
 			// if it's a relative path, let's have it relative to the game 
 			// directory (not the UE executable)
 			if (FPaths::IsRelative(Value))
@@ -116,20 +137,20 @@ FNativeCodeGenCommandlineParams::FNativeCodeGenCommandlineParams(const TArray<FS
 
 			if (FileManager.DirectoryExists(*Value) || FileManager.MakeDirectory(*Value))
 			{
-				ModuleOutputDir = Value;
+				OutputDir = Value;
 			}
 			else
 			{
-				UE_LOG(LogNativeCodeGenCommandline, Warning, TEXT("'%s' doesn't appear to be a valid output directory, defaulting to: '%s'"), *Value, *DefaultModulePath);
+				UE_LOG(LogNativeCodeGenCommandline, Warning, TEXT("'%s' doesn't appear to be a valid output directory, reverting to the default."), *Value);
 			}
 		}
-		else if (!Switch.Compare(TEXT("moduleName"), ESearchCase::IgnoreCase))
+		else if (!Switch.Compare(TEXT("pluginName"), ESearchCase::IgnoreCase))
 		{
-			ModuleName = Value;
+			PluginName = Value;
 		}
-		else if (!Switch.Compare(TEXT("wipe"), ESearchCase::IgnoreCase))
+		else if (!Switch.Compare(TEXT("noWipe"), ESearchCase::IgnoreCase))
 		{
-			bWipeRequested = true;
+			bWipeRequested = false;
 		}
 		else if (!Switch.Compare(TEXT("preview"), ESearchCase::IgnoreCase))
 		{
@@ -137,7 +158,7 @@ FNativeCodeGenCommandlineParams::FNativeCodeGenCommandlineParams(const TArray<FS
 		}
 		else if (!Switch.Compare(TEXT("manifest"), ESearchCase::IgnoreCase))
 		{
-			FPaths::NormalizeDirectoryName(Value);
+			SanitizePathParam(Value);
 
 			if (FileManager.FileExists(*Value))
 			{
@@ -162,8 +183,8 @@ FNativeCodeGenCommandlineParams::FNativeCodeGenCommandlineParams(const TArray<FS
 					{
 						Value = FPaths::ConvertRelativePathToFull(Value);
 					}
-					UE_LOG( LogNativeCodeGenCommandline, Warning, TEXT("'%s' doesn't appear to be a valid manifest file path, defaulting to: '%s'."), *Value,
-						*FPaths::Combine(*DefaultModulePath, *FBlueprintNativeCodeGenManifest::GetDefaultFilename()) );
+					UE_LOG( LogNativeCodeGenCommandline, Warning, TEXT("'%s' doesn't appear to be a valid manifest file name/path, defaulting to: '%s'."), 
+						*Value,	*FBlueprintNativeCodeGenManifest::GetDefaultFilename() );
 				}
 			}
 		}
@@ -174,16 +195,22 @@ FNativeCodeGenCommandlineParams::FNativeCodeGenCommandlineParams(const TArray<FS
 	}
 
 	bool const bUtilizeExistingManifest = !ManifestFilePath.IsEmpty() && !bWipeRequested && FileManager.FileExists(*ManifestFilePath);
-	// an existing manifest would specify where to put the module
+	// an existing manifest would specify where to put the plugin
 	if (!bUtilizeExistingManifest)
 	{
-		if (ModuleName.IsEmpty())
+		if (PluginName.IsEmpty())
 		{
-			ModuleName = NativeCodeGenCommandlineParamsImpl::DefaultModuleName;
+			PluginName = NativeCodeGenCommandlineParamsImpl::DefaultPluginName;
 		}
-		if (ModuleOutputDir.IsEmpty())
+		if (OutputDir.IsEmpty())
+		{
+			const FString DefaultOutputPath = FPaths::Combine(*FPaths::Combine(*FPaths::GameIntermediateDir(), TEXT("Plugins")), *PluginName);
+			OutputDir = DefaultOutputPath;
+		}
+
+		if ( PluginName.Compare(FPaths::GetBaseFilename(OutputDir)) )
 		{
-			ModuleOutputDir = DefaultModulePath;
+			OutputDir = FPaths::Combine(*OutputDir, *PluginName);
 		}
 	}	
 }
diff --git a/Engine/Source/Developer/BlueprintNativeCodeGen/Private/NativeCodeGenCommandlineParams.h b/Engine/Source/Developer/BlueprintNativeCodeGen/Private/NativeCodeGenCommandlineParams.h
--- a/Engine/Source/Developer/BlueprintNativeCodeGen/Private/NativeCodeGenCommandlineParams.h
+++ b/Engine/Source/Developer/BlueprintNativeCodeGen/Private/NativeCodeGenCommandlineParams.h
@@ -15,8 +15,8 @@ public:
 	bool            bHelpRequested : 1;
 	TArray<FString> WhiteListedAssetPaths;
 	TArray<FString> BlackListedAssetPaths;
-	FString         ModuleOutputDir;
-	FString			ModuleName;
+	FString         OutputDir;
+	FString			PluginName;
 	FString			ManifestFilePath;
 	bool            bWipeRequested : 1;
 	bool            bPreviewRequested : 1;
