(WIP) Setting up the code generation backend so we can override the target PCH.

diff --git a/Engine/Source/Developer/BlueprintCompilerCppBackend/Private/BlueprintCompilerCppBackendBase.cpp b/Engine/Source/Developer/BlueprintCompilerCppBackend/Private/BlueprintCompilerCppBackendBase.cpp
--- a/Engine/Source/Developer/BlueprintCompilerCppBackend/Private/BlueprintCompilerCppBackendBase.cpp
+++ b/Engine/Source/Developer/BlueprintCompilerCppBackend/Private/BlueprintCompilerCppBackendBase.cpp
@@ -3,7 +3,6 @@
 #include "BlueprintCompilerCppBackendModulePrivatePCH.h"
 #include "BlueprintCompilerCppBackendBase.h"
 #include "BlueprintCompilerCppBackendUtils.h"
-#include "IBlueprintCompilerCppBackendModule.h" // for GetBaseFilename()
 
 void FBlueprintCompilerCppBackendBase::EmitStructProperties(FEmitterLocalContext& EmitterContext, UStruct* SourceClass)
 {
@@ -593,7 +592,7 @@ void FBlueprintCompilerCppBackendBase::EmitFileBeginning(const FString& CleanNam
 	{
 		Dst.AddLine(FString::Printf(TEXT("#include \"%s%s\""), Message, bAddDotH ? TEXT(".h") : TEXT("")));
 	};
-	EmitIncludeHeader(EmitterContext.Body, FApp::GetGameName(), true);
+	EmitIncludeHeader(EmitterContext.Body, *FEmitHelper::GetPCHFilename(), false);
 	EmitIncludeHeader(EmitterContext.Body, *CleanName, true);
 	EmitIncludeHeader(bIncludeCodeHelpersInHeader ? EmitterContext.Header : EmitterContext.Body, TEXT("GeneratedCodeHelpers"), true);
 
@@ -621,7 +620,7 @@ void FBlueprintCompilerCppBackendBase::EmitFileBeginning(const FString& CleanNam
 				AlreadyIncluded.Add(Name, &bAlreadyIncluded);
 				if (!bAlreadyIncluded)
 				{
-					const FString GeneratedFilename = IBlueprintCompilerCppBackendModule::GetBaseFilename(Field);
+					const FString GeneratedFilename = FEmitHelper::GetBaseFilename(Field);
 					EmitIncludeHeader(Dst, *GeneratedFilename, true);
 				}
 			}
diff --git a/Engine/Source/Developer/BlueprintCompilerCppBackend/Private/BlueprintCompilerCppBackendModule.cpp b/Engine/Source/Developer/BlueprintCompilerCppBackend/Private/BlueprintCompilerCppBackendModule.cpp
--- a/Engine/Source/Developer/BlueprintCompilerCppBackend/Private/BlueprintCompilerCppBackendModule.cpp
+++ b/Engine/Source/Developer/BlueprintCompilerCppBackend/Private/BlueprintCompilerCppBackendModule.cpp
@@ -8,35 +8,34 @@
 
 class FBlueprintCompilerCppBackendModule : public IBlueprintCompilerCppBackendModule
 {
-	/** IModuleInterface implementation */
-	virtual void StartupModule() override;
-	virtual void ShutdownModule() override;
-	virtual IBlueprintCompilerCppBackend* Create()  override;
+public:
+	//~ Begin IBlueprintCompilerCppBackendModuleInterface interface
+	virtual IBlueprintCompilerCppBackend* Create() override;
+	//~ End IBlueprintCompilerCppBackendModuleInterface interface
+
+	//~ Begin IBlueprintCompilerCppBackendModule interface
+	FString ConstructBaseFilename(const UObject* AssetObj) override;
+	virtual FPCHFilenameQuery& OnPCHFilenameQuery() override;
+	//~ End IBlueprintCompilerCppBackendModule interface
+
+private: 
+	FPCHFilenameQuery PCHFilenameQuery;
 };
 
-IMPLEMENT_MODULE(FBlueprintCompilerCppBackendModule, BlueprintCompilerCppBackend)
-
-
-
-void FBlueprintCompilerCppBackendModule::StartupModule()
-{
-	// This code will execute after your module is loaded into memory (but after global variables are initialized, of course.)
-}
-
-
-void FBlueprintCompilerCppBackendModule::ShutdownModule()
-{
-	// This function may be called during shutdown to clean up your module.  For modules that support dynamic reloading,
-	// we call this function before unloading the module.
-}
-
 IBlueprintCompilerCppBackend* FBlueprintCompilerCppBackendModule::Create()
 {
 	return new FBlueprintCompilerCppBackend();
 }
 
-FString IBlueprintCompilerCppBackendModule::GetBaseFilename(const UObject* AssetObj)
+FString FBlueprintCompilerCppBackendModule::ConstructBaseFilename(const UObject* AssetObj)
 {
 	// use the same function that the backend uses for #includes
 	return FEmitHelper::GetBaseFilename(AssetObj);
 }
+
+IBlueprintCompilerCppBackendModule::FPCHFilenameQuery& FBlueprintCompilerCppBackendModule::OnPCHFilenameQuery()
+{
+	return PCHFilenameQuery;
+}
+
+IMPLEMENT_MODULE(FBlueprintCompilerCppBackendModule, BlueprintCompilerCppBackend)
diff --git a/Engine/Source/Developer/BlueprintCompilerCppBackend/Private/BlueprintCompilerCppBackendUtils.cpp b/Engine/Source/Developer/BlueprintCompilerCppBackend/Private/BlueprintCompilerCppBackendUtils.cpp
--- a/Engine/Source/Developer/BlueprintCompilerCppBackend/Private/BlueprintCompilerCppBackendUtils.cpp
+++ b/Engine/Source/Developer/BlueprintCompilerCppBackend/Private/BlueprintCompilerCppBackendUtils.cpp
@@ -3,6 +3,7 @@
 #include "BlueprintCompilerCppBackendModulePrivatePCH.h"
 #include "BlueprintCompilerCppBackendUtils.h"
 #include "EdGraphSchema_K2.h"
+#include "IBlueprintCompilerCppBackendModule.h" // for OnPCHFilenameQuery()
 
 FString FEmitterLocalContext::GenerateUniqueLocalName()
 {
@@ -576,6 +577,25 @@ FString FEmitHelper::GetBaseFilename(const UObject* AssetObj)
 	return FPackageName::GetLongPackageAssetName(PackagePath+Postfix);
 }
 
+FString FEmitHelper::GetPCHFilename()
+{
+	FString PCHFilename;
+	FBlueprintCompilerCppBackendModule& BackEndModule = (FBlueprintCompilerCppBackendModule&)FBlueprintCompilerCppBackendModule::Get();
+
+	auto& PchFilenameQuery = BackEndModule.OnPCHFilenameQuery();
+	if (PchFilenameQuery.IsBound())
+	{
+		PCHFilename = PchFilenameQuery.Execute();
+	}
+
+	if (PCHFilename.IsEmpty())
+	{
+		PCHFilename  = FApp::GetGameName();
+		PCHFilename += TEXT(".h");
+	}
+	return PCHFilename;
+}
+
 FString FEmitHelper::EmitUFuntion(UFunction* Function, TArray<FString>* AdditinalMetaData)
 {
 	TArray<FString> Tags = FEmitHelper::FunctionFlagsToTags(Function->FunctionFlags);
diff --git a/Engine/Source/Developer/BlueprintCompilerCppBackend/Private/BlueprintCompilerCppBackendUtils.h b/Engine/Source/Developer/BlueprintCompilerCppBackend/Private/BlueprintCompilerCppBackendUtils.h
--- a/Engine/Source/Developer/BlueprintCompilerCppBackend/Private/BlueprintCompilerCppBackendUtils.h
+++ b/Engine/Source/Developer/BlueprintCompilerCppBackend/Private/BlueprintCompilerCppBackendUtils.h
@@ -228,6 +228,8 @@ struct FEmitHelper
 	static FString GetBaseFilename(const UObject* AssetObj);
 
 	static FString ReplaceConvertedMetaData(UObject* Obj);
+
+	static FString GetPCHFilename();
 };
 
 struct FNonativeComponentData;
diff --git a/Engine/Source/Developer/BlueprintCompilerCppBackend/Public/IBlueprintCompilerCppBackendModule.h b/Engine/Source/Developer/BlueprintCompilerCppBackend/Public/IBlueprintCompilerCppBackendModule.h
--- a/Engine/Source/Developer/BlueprintCompilerCppBackend/Public/IBlueprintCompilerCppBackendModule.h
+++ b/Engine/Source/Developer/BlueprintCompilerCppBackend/Public/IBlueprintCompilerCppBackendModule.h
@@ -7,10 +7,27 @@
 /**
  * The public interface to this module
  */
-class BLUEPRINTCOMPILERCPPBACKEND_API IBlueprintCompilerCppBackendModule : public IBlueprintCompilerCppBackendModuleInterface
+class IBlueprintCompilerCppBackendModule : public IBlueprintCompilerCppBackendModuleInterface
 {
 public:
+	/**
+	 * Aimed to help coordinate #include statements (generated by the backend) 
+	 * and target file paths. A centralized function, that'll generate a 
+	 * source-file name for the supplied asset object.
+	 * 
+	 * @param  AssetObj    The asset you want a source-file name for.
+	 * @return The filename (without an extension) for the target asset.
+	 */
+	virtual FString ConstructBaseFilename(const UObject* AssetObj) = 0;
 
-	static FString GetBaseFilename(const UObject* AssetObj);
+
+	DECLARE_DELEGATE_RetVal(FString, FPCHFilenameQuery);
+	/**
+	 * Provides a hook so that external modules can specify the PCH file that 
+	 * should be used for generated source.
+	 * 
+	 * @return A delegate that will be used to fill out the shared include statement in generated source.
+	 */
+	virtual FPCHFilenameQuery& OnPCHFilenameQuery() = 0;
 };
 
diff --git a/Engine/Source/Developer/BlueprintNativeCodeGen/Private/BlueprintNativeCodeGenManifest.cpp b/Engine/Source/Developer/BlueprintNativeCodeGen/Private/BlueprintNativeCodeGenManifest.cpp
--- a/Engine/Source/Developer/BlueprintNativeCodeGen/Private/BlueprintNativeCodeGenManifest.cpp
+++ b/Engine/Source/Developer/BlueprintNativeCodeGen/Private/BlueprintNativeCodeGenManifest.cpp
@@ -8,7 +8,7 @@
 #include "Serialization/JsonSerializer.h"
 #include "JsonObjectConverter.h"
 #include "Serialization/JsonWriter.h"
-#include "IBlueprintCompilerCppBackendModule.h" // for GetBaseFilename()
+#include "IBlueprintCompilerCppBackendModule.h" // for ConstructBaseFilename()
 
 DEFINE_LOG_CATEGORY_STATIC(LogNativeCodeGenManifest, Log, All);
 
@@ -150,7 +150,8 @@ static FString BlueprintNativeCodeGenManifestImpl::GenerateCppSavePath(const FSt
 //------------------------------------------------------------------------------
 static FString BlueprintNativeCodeGenManifestImpl::GetBaseFilename(const FAssetData& Asset)
 {
-	return IBlueprintCompilerCppBackendModule::GetBaseFilename(Asset.GetAsset());
+	IBlueprintCompilerCppBackendModule& CodeGenBackend = (IBlueprintCompilerCppBackendModule&)IBlueprintCompilerCppBackendModule::Get();
+	return CodeGenBackend.ConstructBaseFilename(Asset.GetAsset());
 }
 
 //------------------------------------------------------------------------------
diff --git a/Engine/Source/Developer/BlueprintNativeCodeGen/Private/NativeCodeGenerationTool.cpp b/Engine/Source/Developer/BlueprintNativeCodeGen/Private/NativeCodeGenerationTool.cpp
--- a/Engine/Source/Developer/BlueprintNativeCodeGen/Private/NativeCodeGenerationTool.cpp
+++ b/Engine/Source/Developer/BlueprintNativeCodeGen/Private/NativeCodeGenerationTool.cpp
@@ -13,7 +13,7 @@
 #include "BlueprintNativeCodeGenUtils.h"
 //#include "Editor/KismetCompiler/Public/BlueprintCompilerCppBackendInterface.h"
 #include "Developer/BlueprintCompilerCppBackend/Public/BlueprintCompilerCppBackendGatherDependencies.h"
-#include "IBlueprintCompilerCppBackendModule.h" // for GetBaseFilename()
+#include "IBlueprintCompilerCppBackendModule.h" // for ConstructBaseFilename()
 #include "KismetCompilerModule.h"
 
 #define LOCTEXT_NAMESPACE "NativeCodeGenerationTool"
@@ -31,7 +31,9 @@ struct FGeneratedCodeData
 		FName GeneratedClassName, SkeletonClassName;
 		InBlueprint.GetBlueprintClassNames(GeneratedClassName, SkeletonClassName);
 		ClassName = GeneratedClassName.ToString();
-		BaseFilename = IBlueprintCompilerCppBackendModule::GetBaseFilename(&InBlueprint);
+
+		IBlueprintCompilerCppBackendModule& CodeGenBackend = (IBlueprintCompilerCppBackendModule&)IBlueprintCompilerCppBackendModule::Get();
+		BaseFilename = CodeGenBackend.ConstructBaseFilename(&InBlueprint);
 
 		GatherUserDefinedDependencies(InBlueprint);
 	}
@@ -125,6 +127,8 @@ struct FGeneratedCodeData
 		FScopedSlowTask SlowTask(WorkParts, LOCTEXT("GeneratingCppFiles", "Generating C++ files.."));
 		SlowTask.MakeDialog();
 
+		IBlueprintCompilerCppBackendModule& CodeGenBackend = (IBlueprintCompilerCppBackendModule&)IBlueprintCompilerCppBackendModule::Get();
+
 		TArray<FString> CreatedFiles;
 		for(auto Obj : DependentObjects)
 		{
@@ -135,7 +139,7 @@ struct FGeneratedCodeData
 			FBlueprintNativeCodeGenUtils::GenerateCppCode(Obj, HeaderSource, CppSource);
 			SlowTask.EnterProgressFrame();
 
-			const FString BackendBaseFilename = IBlueprintCompilerCppBackendModule::GetBaseFilename(Obj);
+			const FString BackendBaseFilename = CodeGenBackend.ConstructBaseFilename(Obj);
 
 			const FString FullHeaderFilename = FPaths::Combine(*HeaderDirPath, *(BackendBaseFilename + TEXT(".h")));
 			const bool bHeaderSaved = FFileHelper::SaveStringToFile(*HeaderSource, *FullHeaderFilename);
@@ -173,7 +177,7 @@ struct FGeneratedCodeData
 
 			SlowTask.EnterProgressFrame();
 
-			const FString BackendBaseFilename = IBlueprintCompilerCppBackendModule::GetBaseFilename(BPGC);
+			const FString BackendBaseFilename = CodeGenBackend.ConstructBaseFilename(BPGC);
 
 			const FString FullHeaderFilename = FPaths::Combine(*HeaderDirPath, *(BackendBaseFilename + TEXT(".h")));
 			const bool bHeaderSaved = FFileHelper::SaveStringToFile(HeaderSource, *FullHeaderFilename);
