Copying //UE4/Dev-Blueprints to Dev-Main (//UE4/Dev-Main) @ 2781164

diff --git a/Engine/Source/Developer/BlueprintCompilerCppBackend/Private/BlueprintCompilerCppBackendUtils.cpp b/Engine/Source/Developer/BlueprintCompilerCppBackend/Private/BlueprintCompilerCppBackendUtils.cpp
--- a/Engine/Source/Developer/BlueprintCompilerCppBackend/Private/BlueprintCompilerCppBackendUtils.cpp
+++ b/Engine/Source/Developer/BlueprintCompilerCppBackend/Private/BlueprintCompilerCppBackendUtils.cpp
@@ -15,8 +15,20 @@ FString FEmitterLocalContext::GenerateUniqueLocalName()
 
 FString FEmitterLocalContext::FindGloballyMappedObject(const UObject* Object, const UClass* ExpectedClass, bool bLoadIfNotFound, bool bTryUsedAssetsList)
 {
-	UClass* ActualClass = Cast<UClass>(Dependencies.GetOriginalStruct());
-	if (ActualClass && Object && Object->IsIn(ActualClass))
+	UClass* ActualClass = Cast<UClass>(Dependencies.GetActualStruct());
+	UClass* OriginalActualClass = Dependencies.FindOriginalClass(ActualClass);
+	UClass* OuterClass = Object ? Cast<UClass>(Object->GetOuter()) : nullptr;	// SCS component templates will have an Outer that equates to their owning BPGC; since they're not currently DSOs, we have to special-case them.
+	
+	auto ClassString = [&]() -> FString
+	{
+		const UClass* ObjectClassToUse = ExpectedClass ? ExpectedClass : GetFirstNativeOrConvertedClass(Object->GetClass());
+		return FEmitHelper::GetCppName(ObjectClassToUse);
+	};
+	
+	if (ActualClass && Object 
+		&& (Object->IsIn(ActualClass) 
+		|| Object->IsIn(ActualClass->GetDefaultObject(false))
+		|| (OuterClass && ActualClass->IsChildOf(OuterClass)) ))
 	{
 		if (const FString* NamePtr = (CurrentCodeType == EGeneratedCodeType::SubobjectsOfClass) ? ClassSubobjectsMap.Find(Object) : nullptr)
 		{
@@ -28,49 +40,72 @@ FString FEmitterLocalContext::FindGloballyMappedObject(const UObject* Object, co
 			return *NamePtr;
 		}
 
-		const UClass* ObjectClassToUse = ExpectedClass ? ExpectedClass : GetFirstNativeOrConvertedClass(Object->GetClass());
-		const FString ClassString = FEmitHelper::GetCppName(ObjectClassToUse);
-
 		int32 ObjectsCreatedPerClassIdx = MiscConvertedSubobjects.IndexOfByKey(Object);
 		if (INDEX_NONE != ObjectsCreatedPerClassIdx)
 		{
 			return FString::Printf(TEXT("CastChecked<%s>(CastChecked<UDynamicClass>(%s::StaticClass())->MiscConvertedSubobjects[%d])")
-				, *ClassString, *FEmitHelper::GetCppName(ActualClass), ObjectsCreatedPerClassIdx);
+				, *ClassString(), *FEmitHelper::GetCppName(ActualClass), ObjectsCreatedPerClassIdx);
 		}
 
 		ObjectsCreatedPerClassIdx = ObjectsCreatedPerClassIdx = DynamicBindingObjects.IndexOfByKey(Object);
 		if (INDEX_NONE != ObjectsCreatedPerClassIdx)
 		{
 			return FString::Printf(TEXT("CastChecked<%s>(CastChecked<UDynamicClass>(%s::StaticClass())->DynamicBindingObjects[%d])")
-				, *ClassString, *FEmitHelper::GetCppName(ActualClass), ObjectsCreatedPerClassIdx);
+				, *ClassString(), *FEmitHelper::GetCppName(ActualClass), ObjectsCreatedPerClassIdx);
 		}
 
 		ObjectsCreatedPerClassIdx = ComponentTemplates.IndexOfByKey(Object);
 		if (INDEX_NONE != ObjectsCreatedPerClassIdx)
 		{
 			return FString::Printf(TEXT("CastChecked<%s>(CastChecked<UDynamicClass>(%s::StaticClass())->ComponentTemplates[%d])")
-				, *ClassString, *FEmitHelper::GetCppName(ActualClass), ObjectsCreatedPerClassIdx);
+				, *ClassString(), *FEmitHelper::GetCppName(ActualClass), ObjectsCreatedPerClassIdx);
 		}
 
 		ObjectsCreatedPerClassIdx = Timelines.IndexOfByKey(Object);
 		if (INDEX_NONE != ObjectsCreatedPerClassIdx)
 		{
 			return FString::Printf(TEXT("CastChecked<%s>(CastChecked<UDynamicClass>(%s::StaticClass())->Timelines[%d])")
-				, *ClassString, *FEmitHelper::GetCppName(ActualClass), ObjectsCreatedPerClassIdx);
+				, *ClassString(), *FEmitHelper::GetCppName(ActualClass), ObjectsCreatedPerClassIdx);
 		}
 
 		if ((CurrentCodeType == EGeneratedCodeType::SubobjectsOfClass) || (CurrentCodeType == EGeneratedCodeType::CommonConstructor))
 		{
-			if (Object == ActualClass->GetDefaultObject(false))
+			if ((Object == ActualClass->GetDefaultObject(false)) || (Object == OriginalActualClass->GetDefaultObject(false)))
 			{
 				return TEXT("this");
 			}
 		}
 	}
 
-	if (ActualClass && (Object == ActualClass))
+	auto CastCustomClass = [&](const FString& InResult)->FString
 	{
-		return TEXT("GetClass()");
+		if (ExpectedClass && !UClass::StaticClass()->IsChildOf(ExpectedClass))
+		{
+			return FString::Printf(TEXT("Cast<%s>(%s)"), *ClassString(), *InResult);
+		}
+		return InResult;
+	};
+
+	if (ActualClass && ((Object == ActualClass) || (Object == OriginalActualClass)))
+	{
+		return CastCustomClass(TEXT("GetClass()"));
+	}
+
+	{
+		auto Field = Cast<UField>(Object);
+		auto FieldOwnerStruct = Field ? Field->GetOwnerStruct() : nullptr;
+		if (FieldOwnerStruct && (Field != FieldOwnerStruct))
+		{
+			ensure(Field == FindField<UField>(FieldOwnerStruct, Field->GetFName()));
+			const FString MappedOwner = FindGloballyMappedObject(FieldOwnerStruct, UStruct::StaticClass(), bLoadIfNotFound, bTryUsedAssetsList);
+			if (!MappedOwner.IsEmpty() && ensure(MappedOwner != TEXT("nullptr")))
+			{
+				return FString::Printf(TEXT("FindFieldChecked<%s>(%s, TEXT(\"%s\"))")
+					, *FEmitHelper::GetCppName(Field->GetClass())
+					, *MappedOwner
+					, *Field->GetName());
+			}
+		}
 	}
 
 	if (auto ObjClass = Cast<UClass>(Object))
@@ -78,7 +113,7 @@ FString FEmitterLocalContext::FindGloballyMappedObject(const UObject* Object, co
 		auto BPGC = Cast<UBlueprintGeneratedClass>(ObjClass);
 		if (ObjClass->HasAnyClassFlags(CLASS_Native) || (BPGC && Dependencies.WillClassBeConverted(BPGC)))
 		{
-			return FString::Printf(TEXT("%s::StaticClass()"), *FEmitHelper::GetCppName(ObjClass, true));
+			return CastCustomClass(FString::Printf(TEXT("%s::StaticClass()"), *FEmitHelper::GetCppName(ObjClass, true)));
 		}
 	}
 
@@ -103,16 +138,13 @@ FString FEmitterLocalContext::FindGloballyMappedObject(const UObject* Object, co
 	ensure(!bLoadIfNotFound || Object);
 	if (Object && (bLoadIfNotFound || bTryUsedAssetsList))
 	{
-		const UClass* ObjectClassToUse = ExpectedClass ? ExpectedClass : GetFirstNativeOrConvertedClass(Object->GetClass());
-		const FString ClassString = FEmitHelper::GetCppName(ObjectClassToUse);
-
 		if (bTryUsedAssetsList)
 		{
 			const int32 AssetIndex = Dependencies.Assets.IndexOfByKey(Object);
 			if (INDEX_NONE != AssetIndex)
 			{
 				return FString::Printf(TEXT("CastChecked<%s>(CastChecked<UDynamicClass>(%s::StaticClass())->UsedAssets[%d])")
-					, *ClassString
+					, *ClassString()
 					, *FEmitHelper::GetCppName(ActualClass)
 					, AssetIndex);
 			}
@@ -121,7 +153,7 @@ FString FEmitterLocalContext::FindGloballyMappedObject(const UObject* Object, co
 		if (bLoadIfNotFound)
 		{
 			return FString::Printf(TEXT("LoadObject<%s>(nullptr, TEXT(\"%s\"))")
-				, *ClassString
+				, *ClassString()
 				, *(Object->GetPathName().ReplaceCharWithEscapedChar()));
 		}
 	}
@@ -142,19 +174,6 @@ FString FEmitterLocalContext::ExportTextItem(const UProperty* Property, const vo
 		return FString::Printf(TEXT("%s()"), *TypeText);
 	}
 
-	auto ObjectPropertyBase = Cast<const UObjectPropertyBase>(Property);
-	if (ObjectPropertyBase && Property->IsA<UAssetObjectProperty>())
-	{
-		if (UObject* Object = ObjectPropertyBase->GetObjectPropertyValue(PropertyValue))
-		{
-			const UClass* ActualClass = GetFirstNativeOrConvertedClass(ObjectPropertyBase->PropertyClass);
-			return FString::Printf(TEXT("LoadObject<%s>(nullptr, TEXT(\"%s\"))")
-				, *FEmitHelper::GetCppName(ActualClass)
-				, *(Object->GetPathName().ReplaceCharWithEscapedChar()));
-		}
-		return TEXT("nullptr");
-	}
-
 	FString ValueStr;
 	Property->ExportTextItem(ValueStr, PropertyValue, PropertyValue, nullptr, EPropertyPortFlags::PPF_ExportCpp);
 	return ValueStr;
@@ -171,25 +190,17 @@ FString FEmitterLocalContext::ExportCppDeclaration(const UProperty* Property, EE
 	auto GetActualNameCPP = [&](const UObjectPropertyBase* ObjectPropertyBase, UClass* InActualClass)
 	{
 		auto BPGC = Cast<UBlueprintGeneratedClass>(InActualClass);
-		const bool bConverted = BPGC && Dependencies.WillClassBeConverted(BPGC);
 		if (BPGC )
 		{
-			if (!bConverted)
-			{
-				const bool bIsParameter = (DeclarationType == EExportedDeclaration::Parameter) || (DeclarationType == EExportedDeclaration::MacroParameter);
-				const uint32 LocalExportCPPFlags = ExportCPPFlags | (bIsParameter ? CPPF_ArgumentOrReturnValue : 0);
-				UClass* NativeType = GetFirstNativeOrConvertedClass(BPGC);
-				ActualCppType = ObjectPropertyBase->GetCPPTypeCustom(&ActualExtendedType, LocalExportCPPFlags, NativeType);
-				ActualCppTypePtr = &ActualCppType;
-				if (!ActualExtendedType.IsEmpty())
-				{
-					ActualExtendedTypePtr = &ActualExtendedType;
-				}
-			}
-			else
+			const bool bIsParameter = (DeclarationType == EExportedDeclaration::Parameter) || (DeclarationType == EExportedDeclaration::MacroParameter);
+			const uint32 LocalExportCPPFlags = ExportCPPFlags | (bIsParameter ? CPPF_ArgumentOrReturnValue : 0);
+			UClass* NativeType = GetFirstNativeOrConvertedClass(BPGC);
+			check(NativeType);
+			ActualCppType = ObjectPropertyBase->GetCPPTypeCustom(&ActualExtendedType, LocalExportCPPFlags, FEmitHelper::GetCppName(NativeType));
+			ActualCppTypePtr = &ActualCppType;
+			if (!ActualExtendedType.IsEmpty())
 			{
-				ActualCppType = FEmitHelper::GetCppName(ObjectPropertyBase->PropertyClass, false) + TEXT("*");
-				ActualCppTypePtr = &ActualCppType;
+				ActualExtendedTypePtr = &ActualExtendedType;
 			}
 		}
 	};
@@ -218,6 +229,15 @@ FString FEmitterLocalContext::ExportCppDeclaration(const UProperty* Property, EE
 		ActualCppType = FEmitHelper::GetCppName(StructProperty->Struct, false);
 		ActualCppTypePtr = &ActualCppType;
 	}
+	else if (auto SCDelegateProperty = Cast<const UDelegateProperty>(Property))
+	{
+		const FString* SCDelegateTypeName = MCDelegateSignatureToSCDelegateType.Find(SCDelegateProperty->SignatureFunction);
+		if (SCDelegateTypeName)
+		{
+			ActualCppType = *SCDelegateTypeName;
+			ActualCppTypePtr = &ActualCppType;
+		}
+	}
 
 	if (ArrayProperty)
 	{
@@ -278,7 +298,8 @@ FString FEmitHelper::GetCppName(const UField* Field, bool bUInterface)
 			if ((Cast<UBlueprintGeneratedClass>(Owner) ||
 				!Owner->IsNative()))
 			{
-				return ::UnicodeToCPPIdentifier(AsProperty->GetName(), AsProperty->HasAnyPropertyFlags(CPF_Deprecated), TEXT("bpv__"));
+				const bool bIsParameter = AsProperty->HasAnyPropertyFlags(CPF_Parm);
+				return ::UnicodeToCPPIdentifier(AsProperty->GetName(), AsProperty->HasAnyPropertyFlags(CPF_Deprecated), bIsParameter ? TEXT("bpp__") : TEXT("bpv__"));
 			}
 		}
 		return AsProperty->GetNameCPP();
@@ -323,16 +344,24 @@ FString FEmitHelper::HandleRepNotifyFunc(const UProperty* Property)
 	return FString();
 }
 
-bool FEmitHelper::MetaDataCanBeNative(const FName MetaDataName)
+bool FEmitHelper::MetaDataCanBeNative(const FName MetaDataName, const UField* Field)
 {
 	if (MetaDataName == TEXT("ModuleRelativePath"))
 	{
 		return false;
 	}
+	if (const UFunction* Function = Cast<const UFunction>(Field))
+	{
+		const UProperty* Param = Function->FindPropertyByName(MetaDataName);
+		if (Param && Param->HasAnyPropertyFlags(CPF_Parm))
+		{
+			return false;
+		}
+	}
 	return true;
 }
 
-FString FEmitHelper::HandleMetaData(const UField* Field, bool AddCategory, TArray<FString>* AdditinalMetaData)
+FString FEmitHelper::HandleMetaData(const UField* Field, bool AddCategory, const TArray<FString>* AdditinalMetaData)
 {
 	FString MetaDataStr;
 
@@ -344,13 +373,13 @@ FString FEmitHelper::HandleMetaData(const UField* Field, bool AddCategory, TArra
 	{
 		for (auto& Pair : *ValuesMap)
 		{
-			if (!MetaDataCanBeNative(Pair.Key))
+			if (!MetaDataCanBeNative(Pair.Key, Field))
 			{
 				continue;
 			}
 			if (!Pair.Value.IsEmpty())
 			{
-				FString Value = Pair.Value.Replace(TEXT("\n"), TEXT(""));
+				FString Value = Pair.Value.Replace(TEXT("\n"), TEXT("")).ReplaceCharWithEscapedChar();
 				MetaDataStrings.Emplace(FString::Printf(TEXT("%s=\"%s\""), *Pair.Key.ToString(), *Value));
 			}
 			else
@@ -369,7 +398,7 @@ FString FEmitHelper::HandleMetaData(const UField* Field, bool AddCategory, TArra
 	}
 	if (Field)
 	{
-		MetaDataStrings.Emplace(FString::Printf(TEXT("OverrideNativeName=\"%s\""), *Field->GetName()));
+		MetaDataStrings.Emplace(FString::Printf(TEXT("OverrideNativeName=\"%s\""), *Field->GetName().ReplaceCharWithEscapedChar()));
 	}
 	MetaDataStrings.Remove(FString());
 	if (MetaDataStrings.Num())
@@ -471,10 +500,8 @@ TArray<FString> FEmitHelper::FunctionFlagsToTags(uint64 Flags)
 
 	//Pointless: BlueprintNativeEvent, BlueprintImplementableEvent
 	//Pointless: CustomThunk
-
-	//TODO: SealedEvent
-	//TODO: Unreliable
-	//TODO: ServiceRequest, ServiceResponse
+	//Pointless: ServiceRequest, ServiceResponse - only useful for native UFunctions, they're for serializing to json
+	//Pointless: SealedEvent
 
 	HANDLE_CPF_TAG(TEXT("Exec"), FUNC_Exec)
 	HANDLE_CPF_TAG(TEXT("Server"), FUNC_Net | FUNC_NetServer)
@@ -487,6 +514,11 @@ TArray<FString> FEmitHelper::FunctionFlagsToTags(uint64 Flags)
 	HANDLE_CPF_TAG(TEXT("BlueprintCosmetic"), FUNC_BlueprintCosmetic)
 	HANDLE_CPF_TAG(TEXT("WithValidation"), FUNC_NetValidate)
 
+	if (HasAllFlags(Flags, FUNC_Net) && !HasAllFlags(Flags, FUNC_NetReliable))
+	{
+		Tags.Emplace(TEXT("Unreliable"));
+	}
+
 	return Tags;
 }
 
@@ -553,18 +585,13 @@ FString FEmitHelper::GenerateReplaceConvertedMD(UObject* Obj)
 
 FString FEmitHelper::GetBaseFilename(const UObject* AssetObj)
 {
-	FString PackagePath = AssetObj->GetOutermost()->GetPathName();
+	FString AssetName = FPackageName::GetLongPackageAssetName(AssetObj->GetOutermost()->GetPathName());
 	// We have to sanitize the package path because UHT is going to generate header guards (preprocessor symbols)
 	// based on the filename. I'm also not interested in exploring the depth of unicode filename support in UHT,
 	// UBT, and our various c++ toolchains, so this logic is pretty aggressive:
 	FString Postfix = TEXT("__pf");
-	for (auto& Char : PackagePath)
+	for (auto& Char : AssetName)
 	{
-		if (Char == '/' || Char == '\\')
-		{
-			continue;
-		}
-
 		if (!IsValidCPPIdentifierChar(Char))
 		{
 			// deterministically map char to a valid ascii character, we have 63 characters available (aA-zZ, 0-9, and _)
@@ -574,7 +601,7 @@ FString FEmitHelper::GetBaseFilename(const UObject* AssetObj)
 		}
 	}
 
-	return FPackageName::GetLongPackageAssetName(PackagePath+Postfix);
+	return AssetName + Postfix;
 }
 
 FString FEmitHelper::GetPCHFilename()
@@ -588,26 +615,21 @@ FString FEmitHelper::GetPCHFilename()
 		PCHFilename = PchFilenameQuery.Execute();
 	}
 
-	if (PCHFilename.IsEmpty())
-	{
-		PCHFilename  = FApp::GetGameName();
-		PCHFilename += TEXT(".h");
-	}
 	return PCHFilename;
 }
 
-FString FEmitHelper::EmitUFuntion(UFunction* Function, TArray<FString>* AdditinalMetaData)
+FString FEmitHelper::GetGameMainHeaderFilename()
 {
-	TArray<FString> Tags = FEmitHelper::FunctionFlagsToTags(Function->FunctionFlags);
+	return FString::Printf(TEXT("%s.h"), FApp::GetGameName());
+}
 
-	auto FunctionOwnerClass = Function->GetOwnerClass();
-	if (FunctionOwnerClass->IsChildOf<UInterface>())
-	{
-		Tags.Emplace(TEXT("BlueprintNativeEvent"));
-	}
+FString FEmitHelper::EmitUFuntion(UFunction* Function, const TArray<FString>& AdditionalTags, const TArray<FString>& AdditinalMetaData)
+{
+	TArray<FString> Tags = FEmitHelper::FunctionFlagsToTags(Function->FunctionFlags);
 
+	Tags.Append(AdditionalTags);
 	const bool bMustHaveCategory = (Function->FunctionFlags & (FUNC_BlueprintCallable | FUNC_BlueprintPure)) != 0;
-	Tags.Emplace(FEmitHelper::HandleMetaData(Function, bMustHaveCategory, AdditinalMetaData));
+	Tags.Emplace(FEmitHelper::HandleMetaData(Function, bMustHaveCategory, &AdditinalMetaData));
 	Tags.Remove(FString());
 
 	FString AllTags;
@@ -647,24 +669,26 @@ int32 FEmitHelper::ParseDelegateDetails(FEmitterLocalContext& EmitterContext, UF
 	return ParameterNum;
 }
 
+void FEmitHelper::EmitSinglecastDelegateDeclarations_Inner(FEmitterLocalContext& EmitterContext, UFunction* Signature, const FString& TypeName)
+{
+	check(Signature);
+	FString ParamNumberStr, Parameters;
+	ParseDelegateDetails(EmitterContext, Signature, Parameters, ParamNumberStr);
+	EmitterContext.Header.AddLine(FString::Printf(TEXT("DECLARE_DYNAMIC_DELEGATE%s(%s%s);"), *ParamNumberStr, *TypeName, *Parameters));
+}
+
 void FEmitHelper::EmitSinglecastDelegateDeclarations(FEmitterLocalContext& EmitterContext, const TArray<UDelegateProperty*>& Delegates)
 {
 	for (auto It : Delegates)
 	{
 		check(It);
-		auto Signature = It->SignatureFunction;
-		check(Signature);
-
-		FString ParamNumberStr, Parameters;
-		ParseDelegateDetails(EmitterContext, Signature, Parameters, ParamNumberStr);
-
 		const uint32 LocalExportCPPFlags = EPropertyExportCPPFlags::CPPF_CustomTypeName
 			| EPropertyExportCPPFlags::CPPF_NoConst
 			| EPropertyExportCPPFlags::CPPF_NoRef
 			| EPropertyExportCPPFlags::CPPF_NoStaticArray
 			| EPropertyExportCPPFlags::CPPF_BlueprintCppBackend;
 		const FString TypeName = EmitterContext.ExportCppDeclaration(It, EExportedDeclaration::Parameter, LocalExportCPPFlags, true);
-		EmitterContext.Header.AddLine(FString::Printf(TEXT("DECLARE_DYNAMIC_DELEGATE%s(%s%s);"), *ParamNumberStr, *TypeName, *Parameters));
+		EmitSinglecastDelegateDeclarations_Inner(EmitterContext, It->SignatureFunction, TypeName);
 	}
 }
 
@@ -705,7 +729,7 @@ void FEmitHelper::EmitLifetimeReplicatedPropsImpl(FEmitterLocalContext& EmitterC
 				EmitterContext.AddLine(TEXT("Super::GetLifetimeReplicatedProps(OutLifetimeProps);"));
 				bFunctionInitilzed = true;
 			}
-			EmitterContext.AddLine(FString::Printf(TEXT("DOREPLIFETIME(%s, %s);"), *CppClassName, *FEmitHelper::GetCppName(*It)));
+			EmitterContext.AddLine(FString::Printf(TEXT("DOREPLIFETIME_DIFFNAMES(%s, %s, FName(TEXT(\"%s\")));"), *CppClassName, *FEmitHelper::GetCppName(*It), *(It->GetName())));
 		}
 	}
 	if (bFunctionInitilzed)
@@ -745,9 +769,12 @@ FString FEmitHelper::LiteralTerm(FEmitterLocalContext& EmitterContext, const FEd
 			// @note: We have to default to the zeroth entry because there may not be a symbol associated with the 
 			// last element (UHT generates a MAX entry for UEnums, even if the user did not declare them, but UHT 
 			// does not generate a symbol for said entry.
-			return FString::Printf(TEXT("%s::%s"), *FEmitHelper::GetCppName(TypeEnum), CustomValue.IsEmpty()
-				? *TypeEnum->GetEnumName(0)
-				: *CustomValue);
+			if (CustomValue.Contains(TEXT("::")))
+			{
+				return CustomValue;
+			}
+			return FString::Printf(TEXT("%s::%s"), *FEmitHelper::GetCppName(TypeEnum)
+				, CustomValue.IsEmpty() ? *TypeEnum->GetEnumName(0) : *CustomValue);
 		}
 		else
 		{
@@ -825,9 +852,26 @@ FString FEmitHelper::LiteralTerm(FEmitterLocalContext& EmitterContext, const FEd
 			{
 				FStructOnScope StructOnScope(StructType);
 
-				FStringOutputDevice ImportError;
+				class FImportTextErrorContext : public FStringOutputDevice
+				{
+				public:
+					int32 NumErrors;
+
+					FImportTextErrorContext() : FStringOutputDevice(), NumErrors(0) {}
+
+					virtual void Serialize(const TCHAR* V, ELogVerbosity::Type Verbosity, const class FName& Category) override
+					{
+						if (ELogVerbosity::Type::Error == Verbosity)
+						{
+							NumErrors++;
+						}
+						FStringOutputDevice::Serialize(V, Verbosity, Category);
+					}
+				};
+
+				FImportTextErrorContext ImportError;
 				const auto EndOfParsedBuff = UStructProperty::ImportText_Static(StructType, TEXT("FEmitHelper::LiteralTerm"), *CustomValue, StructOnScope.GetStructMemory(), 0, nullptr, &ImportError);
-				if (!EndOfParsedBuff || !ImportError.IsEmpty())
+				if (!EndOfParsedBuff || ImportError.NumErrors)
 				{
 					UE_LOG(LogK2Compiler, Error, TEXT("FEmitHelper::LiteralTerm cannot parse struct \"%s\" error: %s"), *CustomValue, *ImportError);
 				}
@@ -915,6 +959,99 @@ FString FEmitHelper::LiteralTerm(FEmitterLocalContext& EmitterContext, const FEd
 
 FString FEmitHelper::DefaultValue(FEmitterLocalContext& EmitterContext, const FEdGraphPinType& Type)
 {
+	if (Type.bIsArray)
+	{
+		auto PinTypeToNativeType = [](const FEdGraphPinType& InType) -> FString
+		{
+			// A temp property should be generated?
+
+			auto Schema = GetDefault<UEdGraphSchema_K2>();
+			if (UEdGraphSchema_K2::PC_String == InType.PinCategory)
+			{
+				return TEXT("FString");
+			}
+			else if (UEdGraphSchema_K2::PC_Boolean == InType.PinCategory)
+			{
+				return TEXT("bool");
+			}
+			else if((UEdGraphSchema_K2::PC_Byte == InType.PinCategory) || (UEdGraphSchema_K2::PC_Enum == InType.PinCategory))
+			{
+				if (UEnum* Enum = Cast<UEnum>(InType.PinSubCategoryObject.Get()))
+				{
+					return FEmitHelper::GetCppName(Enum);
+				}
+				return TEXT("uint8");
+			}
+			else if(UEdGraphSchema_K2::PC_Int == InType.PinCategory)
+			{
+				return TEXT("int32");
+			}
+			else if(UEdGraphSchema_K2::PC_Float == InType.PinCategory)
+			{
+				return TEXT("float");
+			}
+			else if(UEdGraphSchema_K2::PC_Float == InType.PinCategory)
+			{
+				return TEXT("float");
+			}
+			else if(UEdGraphSchema_K2::PC_Name == InType.PinCategory)
+			{
+				return TEXT("FName");
+			}
+			else if(UEdGraphSchema_K2::PC_Text == InType.PinCategory)
+			{
+				return TEXT("FText");
+			}
+			else if(UEdGraphSchema_K2::PC_Struct == InType.PinCategory)
+			{
+				if (UScriptStruct* Struct = Cast<UScriptStruct>(InType.PinSubCategoryObject.Get()))
+				{
+					return FEmitHelper::GetCppName(Struct);
+				}
+			}
+			else if(UEdGraphSchema_K2::PC_Class == InType.PinCategory)
+			{
+				if (UClass* Class = Cast<UClass>(InType.PinSubCategoryObject.Get()))
+				{
+					return FString::Printf(TEXT("TSubclassOf<%s>"), *FEmitHelper::GetCppName(Class));
+				}
+			}
+			else if (UEdGraphSchema_K2::PC_AssetClass == InType.PinCategory)
+			{
+				if (UClass* Class = Cast<UClass>(InType.PinSubCategoryObject.Get()))
+				{
+					return FString::Printf(TEXT("TAssetSubclassOf<%s>"), *FEmitHelper::GetCppName(Class));
+				}
+			}
+			else if (UEdGraphSchema_K2::PC_Interface == InType.PinCategory)
+			{
+				if (UClass* Class = Cast<UClass>(InType.PinSubCategoryObject.Get()))
+				{
+					return FString::Printf(TEXT("TScriptInterface<%s>"), *FEmitHelper::GetCppName(Class));
+				}
+			}
+			else if (UEdGraphSchema_K2::PC_Asset == InType.PinCategory)
+			{
+				if (UClass* Class = Cast<UClass>(InType.PinSubCategoryObject.Get()))
+				{
+					return FString::Printf(TEXT("TAssetPtr<%s>"), *FEmitHelper::GetCppName(Class));
+				}
+			}
+			else if (UEdGraphSchema_K2::PC_Object == InType.PinCategory)
+			{
+				if (UClass* Class = Cast<UClass>(InType.PinSubCategoryObject.Get()))
+				{
+					return FString::Printf(TEXT("%s*"), *FEmitHelper::GetCppName(Class));
+				}
+			}
+			UE_LOG(LogK2Compiler, Error, TEXT("FEmitHelper::DefaultValue cannot generate an array type"));
+			return FString{};
+		};
+
+		const FString InnerTypeName = PinTypeToNativeType(Type);
+		return FString::Printf(TEXT("TArray<%s>{}"), *InnerTypeName);
+	}
+
 	return LiteralTerm(EmitterContext, Type, FString(), nullptr);
 }
 
@@ -953,10 +1090,7 @@ bool FEmitHelper::ShouldHandleAsNativeEvent(UFunction* Function)
 	{
 		const uint32 FlagsToCheckMask = EFunctionFlags::FUNC_Event | EFunctionFlags::FUNC_BlueprintEvent | EFunctionFlags::FUNC_Native;
 		const uint32 FlagsToCheck = OriginalFunction->FunctionFlags & FlagsToCheckMask;
-
-		auto OriginalOwner = OriginalFunction->GetOwnerClass();
-		const bool bBPInterface = Cast<UBlueprintGeneratedClass>(OriginalOwner) && OriginalOwner->IsChildOf<UInterface>();
-		return (FlagsToCheck == FlagsToCheckMask) || bBPInterface;
+		return (FlagsToCheck == FlagsToCheckMask);
 	}
 	return false;
 }
@@ -977,9 +1111,14 @@ bool FEmitHelper::ShouldHandleAsImplementableEvent(UFunction* Function)
 
 bool FEmitHelper::GenerateAutomaticCast(FEmitterLocalContext& EmitterContext, const FEdGraphPinType& LType, const FEdGraphPinType& RType, FString& OutCastBegin, FString& OutCastEnd)
 {
+	if ((RType.bIsArray != LType.bIsArray) || (LType.PinCategory != RType.PinCategory))
+	{
+		return false;
+	}
+
 	// BYTE to ENUM cast
 	// ENUM to BYTE cast
-	if ((LType.PinCategory == UEdGraphSchema_K2::PC_Byte) && (RType.PinCategory == UEdGraphSchema_K2::PC_Byte))
+	if ((LType.PinCategory == UEdGraphSchema_K2::PC_Byte) && !RType.bIsArray)
 	{
 		auto LTypeEnum = Cast<UEnum>(LType.PinSubCategoryObject.Get());
 		auto RTypeEnum = Cast<UEnum>(RType.PinSubCategoryObject.Get());
@@ -1002,15 +1141,21 @@ bool FEmitHelper::GenerateAutomaticCast(FEmitterLocalContext& EmitterContext, co
 	// OBJECT to OBJECT
 	if (LType.PinCategory == UEdGraphSchema_K2::PC_Object)
 	{
-		UClass* LClass = Cast<UClass>(LType.PinSubCategoryObject.Get());
+		UClass* LClass = EmitterContext.Dependencies.FindOriginalClass( Cast<UClass>(LType.PinSubCategoryObject.Get()) );
 		LClass = LClass ? EmitterContext.GetFirstNativeOrConvertedClass(LClass) : nullptr;
-		UClass* RClass = Cast<UClass>(RType.PinSubCategoryObject.Get());
-		if (LClass && RClass && LClass->IsChildOf(RClass) && !RClass->IsChildOf(LClass))
+		UClass* RClass = EmitterContext.Dependencies.FindOriginalClass( Cast<UClass>(RType.PinSubCategoryObject.Get()) );
+		if (!RType.bIsArray && LClass && RClass && LClass->IsChildOf(RClass) && !RClass->IsChildOf(LClass))
 		{
 			OutCastBegin = FString::Printf(TEXT("CastChecked<%s>("), *FEmitHelper::GetCppName(LClass));
 			OutCastEnd = TEXT(")");
 			return true;
 		}
+		else if (RType.bIsArray && LClass && RClass && (LClass->IsChildOf(RClass) || RClass->IsChildOf(LClass)))
+		{
+			OutCastBegin = FString::Printf(TEXT("TArrayCaster<%s>("), *FEmitHelper::GetCppName(RClass));
+			OutCastEnd = FString::Printf(TEXT(").Get<%s>()"), *FEmitHelper::GetCppName(LClass));
+			return true;
+		}
 	}
 
 	return false;
@@ -1024,8 +1169,51 @@ FString FEmitHelper::ReplaceConvertedMetaData(UObject* Obj)
 	{
 		TArray<FString> AdditionalMD;
 		AdditionalMD.Add(ReplaceConvertedMD);
-		Result += TEXT(",");
 		Result += FEmitHelper::HandleMetaData(nullptr, false, &AdditionalMD);
 	}
 	return Result;
-}
\ No newline at end of file
+}
+
+FString FEmitHelper::GenerateGetPropertyByName(FEmitterLocalContext& EmitterContext, const UProperty* Property)
+{
+	check(Property);
+
+	const FString PropertyWeakPtrName = EmitterContext.GenerateUniqueLocalName();
+	EmitterContext.AddLine(FString::Printf(TEXT("static TWeakObjectPtr<UProperty> %s{};"), *PropertyWeakPtrName));
+
+	const FString PropertyPtrName = EmitterContext.GenerateUniqueLocalName();
+	EmitterContext.AddLine(FString::Printf(TEXT("const UProperty* %s = %s.Get();"), *PropertyPtrName, *PropertyWeakPtrName));
+	EmitterContext.AddLine(FString::Printf(TEXT("if (nullptr == %s)"), *PropertyPtrName));
+	EmitterContext.AddLine(TEXT("{"));
+	EmitterContext.IncreaseIndent();
+
+	const FString PropertyOwnerStruct = EmitterContext.FindGloballyMappedObject(Property->GetOwnerStruct(), UStruct::StaticClass());
+	EmitterContext.AddLine(FString::Printf(TEXT("%s = (%s)->FindPropertyByName(FName(TEXT(\"%s\")));")
+		, *PropertyPtrName
+		, *PropertyOwnerStruct
+		, *Property->GetName()));
+	EmitterContext.AddLine(FString::Printf(TEXT("check(%s);"), *PropertyPtrName));
+	EmitterContext.AddLine(FString::Printf(TEXT("%s = %s;"), *PropertyWeakPtrName, *PropertyPtrName));
+	EmitterContext.DecreaseIndent();
+	EmitterContext.AddLine(TEXT("}"));
+
+	return PropertyPtrName;
+}
+
+FString FEmitHelper::AccessInaccessibleProperty(FEmitterLocalContext& EmitterContext, const UProperty* Property
+	, const FString& ContextStr, const FString& ContextAdressOp, const FString& StaticArrayIdx)
+{
+	check(Property);
+	const FString PropertyLocalName = GenerateGetPropertyByName(EmitterContext, Property);
+	const uint32 CppTemplateTypeFlags = EPropertyExportCPPFlags::CPPF_CustomTypeName
+		| EPropertyExportCPPFlags::CPPF_NoConst | EPropertyExportCPPFlags::CPPF_NoRef | EPropertyExportCPPFlags::CPPF_NoStaticArray
+		| EPropertyExportCPPFlags::CPPF_BlueprintCppBackend;
+	const FString TypeDeclaration = EmitterContext.ExportCppDeclaration(Property, EExportedDeclaration::Parameter, CppTemplateTypeFlags, true);
+
+	return FString::Printf(TEXT("(*(%s->ContainerPtrToValuePtr<%s>(%s(%s)%s)))")
+		, *PropertyLocalName
+		, *TypeDeclaration
+		, *ContextAdressOp
+		, *ContextStr
+		, *StaticArrayIdx);
+}
diff --git a/Engine/Source/Developer/BlueprintCompilerCppBackend/Private/BlueprintCompilerCppBackendUtils.h b/Engine/Source/Developer/BlueprintCompilerCppBackend/Private/BlueprintCompilerCppBackendUtils.h
--- a/Engine/Source/Developer/BlueprintCompilerCppBackend/Private/BlueprintCompilerCppBackendUtils.h
+++ b/Engine/Source/Developer/BlueprintCompilerCppBackend/Private/BlueprintCompilerCppBackendUtils.h
@@ -65,6 +65,8 @@ public:
 
 	const FGatherConvertedClassDependencies& Dependencies;
 
+	TMap<UFunction*, FString> MCDelegateSignatureToSCDelegateType;
+
 	FEmitterLocalContext(const FGatherConvertedClassDependencies& InDependencies)
 		: CurrentCodeType(EGeneratedCodeType::Regular)
 		, LocalNameIndexMax(0)
@@ -147,7 +149,7 @@ public:
 
 	UClass* GetCurrentlyGeneratedClass() const
 	{
-		return Cast<UClass>(Dependencies.GetOriginalStruct());
+		return Cast<UClass>(Dependencies.GetActualStruct());
 	}
 
 	/** All objects (that can be referenced from other package) that will have a different path in cooked build
@@ -175,7 +177,7 @@ public:
 	}
 
 private:
-	FEmitterLocalContext(const FEmitterLocalContext&);
+	FEmitterLocalContext(const FEmitterLocalContext&) = delete;
 };
 
 struct FEmitHelper
@@ -189,9 +191,9 @@ struct FEmitHelper
 
 	static FString HandleRepNotifyFunc(const UProperty* Property);
 
-	static bool MetaDataCanBeNative(const FName MetaDataName);
+	static bool MetaDataCanBeNative(const FName MetaDataName, const UField* Field);
 
-	static FString HandleMetaData(const UField* Field, bool AddCategory = true, TArray<FString>* AdditinalMetaData = nullptr);
+	static FString HandleMetaData(const UField* Field, bool AddCategory = true, const TArray<FString>* AdditinalMetaData = nullptr);
 
 	static TArray<FString> ProperyFlagsToTags(uint64 Flags, bool bIsClassProperty);
 
@@ -201,12 +203,14 @@ struct FEmitHelper
 
 	static bool IsBlueprintImplementableEvent(uint64 FunctionFlags);
 
-	static FString EmitUFuntion(UFunction* Function, TArray<FString>* AdditinalMetaData = nullptr);
+	static FString EmitUFuntion(UFunction* Function, const TArray<FString>& AdditionalTags, const TArray<FString>& AdditionalMetaData);
 
 	static int32 ParseDelegateDetails(FEmitterLocalContext& EmitterContext, UFunction* Signature, FString& OutParametersMacro, FString& OutParamNumberStr);
 
 	static void EmitSinglecastDelegateDeclarations(FEmitterLocalContext& EmitterContext, const TArray<UDelegateProperty*>& Delegates);
 
+	static void EmitSinglecastDelegateDeclarations_Inner(FEmitterLocalContext& EmitterContext, UFunction* Signature, const FString& TypeName);
+
 	static void EmitMulticastDelegateDeclarations(FEmitterLocalContext& EmitterContext);
 
 	static void EmitLifetimeReplicatedPropsImpl(FEmitterLocalContext& EmitterContext);
@@ -230,6 +234,14 @@ struct FEmitHelper
 	static FString ReplaceConvertedMetaData(UObject* Obj);
 
 	static FString GetPCHFilename();
+
+	static FString GetGameMainHeaderFilename();
+
+	static FString GenerateGetPropertyByName(FEmitterLocalContext& EmitterContext, const UProperty* Property);
+
+	static FString AccessInaccessibleProperty(FEmitterLocalContext& EmitterContext, const UProperty* Property
+		, const FString& ContextStr, const FString& ContextAdressOp, const FString& StaticArrayIdx = FString(TEXT(", 0")));
+
 };
 
 struct FNonativeComponentData;
@@ -280,3 +292,10 @@ struct FBackendHelperUMG
 	static void EmitWidgetInitializationFunctions(FEmitterLocalContext& Context);
 
 };
+
+struct FBackendHelperAnim
+{
+	static FString AddHeaders(UClass* GeneratedClass);
+
+	static void CreateAnimClassData(FEmitterLocalContext& Context);
+};
\ No newline at end of file
diff --git a/Engine/Source/Developer/BlueprintCompilerCppBackend/Private/BlueprintCompilerCppBackendValueHelper.cpp b/Engine/Source/Developer/BlueprintCompilerCppBackend/Private/BlueprintCompilerCppBackendValueHelper.cpp
--- a/Engine/Source/Developer/BlueprintCompilerCppBackend/Private/BlueprintCompilerCppBackendValueHelper.cpp
+++ b/Engine/Source/Developer/BlueprintCompilerCppBackend/Private/BlueprintCompilerCppBackendValueHelper.cpp
@@ -35,7 +35,7 @@ void FEmitDefaultValueHelper::OuterGenerate(FEmitterLocalContext& Context
 
 	if (Property->HasAnyPropertyFlags(CPF_EditorOnly | CPF_Transient))
 	{
-		UE_LOG(LogK2Compiler, Log, TEXT("FEmitDefaultValueHelper Skip EditorOnly or Transient property: %s"), *Property->GetPathName());
+		UE_LOG(LogK2Compiler, Verbose, TEXT("FEmitDefaultValueHelper Skip EditorOnly or Transient property: %s"), *Property->GetPathName());
 		return;
 	}
 
@@ -58,27 +58,12 @@ void FEmitDefaultValueHelper::OuterGenerate(FEmitterLocalContext& Context
 			else if (Property->HasAnyPropertyFlags(CPF_NativeAccessSpecifierPrivate) || (!bAllowProtected && Property->HasAnyPropertyFlags(CPF_NativeAccessSpecifierProtected)))
 			{
 				ensure(EPropertyAccessOperator::None != AccessOperator);
-				
-				auto GenerateGetProperty = [](const UProperty* InProperty) -> FString
-				{
-					check(InProperty);
-					const UStruct* OwnerStruct = InProperty->GetOwnerStruct();
-					const FString OwnerName = FEmitHelper::GetCppName(OwnerStruct);
-					const FString OwnerPath = OwnerName + (OwnerStruct->IsA<UClass>() ? TEXT("::StaticClass()") : TEXT("::StaticStruct()"));
-					return FString::Printf(TEXT("%s->FindPropertyByName(FName(TEXT(\"%s\")))"), *OwnerPath, *FEmitHelper::GetCppName(InProperty));
-				};
-				const FString PropertyObject = GenerateGetProperty(Property); //X::StaticClass()->FindPropertyByName(GET_MEMBER_NAME_CHECKED(X, X))
-
-				const uint32 CppTemplateTypeFlags = EPropertyExportCPPFlags::CPPF_CustomTypeName
-					| EPropertyExportCPPFlags::CPPF_NoConst | EPropertyExportCPPFlags::CPPF_NoRef | EPropertyExportCPPFlags::CPPF_NoStaticArray
-					| EPropertyExportCPPFlags::CPPF_BlueprintCppBackend;
-				const FString TypeDeclaration = Context.ExportCppDeclaration(Property, EExportedDeclaration::Parameter, CppTemplateTypeFlags, true);
-
 				const FString OperatorStr = (EPropertyAccessOperator::Dot == AccessOperator) ? TEXT("&") : TEXT("");
-				const FString ContainerStr = (EPropertyAccessOperator::None == AccessOperator) ? TEXT("this") : FString::Printf(TEXT("%s(%s)"), *OperatorStr, *OuterPath);
+				const FString ContainerStr = (EPropertyAccessOperator::None == AccessOperator) ? TEXT("this") : OuterPath;
+				const FString StaticArrayIndexArg = FString::Printf(TEXT(", %d"), ArrayIndex);
+				const FString GetPtrStr = FEmitHelper::AccessInaccessibleProperty(Context, Property, ContainerStr, OperatorStr, StaticArrayIndexArg);
 				PathToMember = Context.GenerateUniqueLocalName();
-				Context.AddLine(FString::Printf(TEXT("auto& %s = *(%s->ContainerPtrToValuePtr<%s>(%s, %d));")
-					, *PathToMember, *PropertyObject, *TypeDeclaration, *ContainerStr, ArrayIndex));
+				Context.AddLine(FString::Printf(TEXT("auto& %s = %s;"), *PathToMember, *GetPtrStr));
 			}
 			else
 			{
@@ -190,6 +175,8 @@ void FEmitDefaultValueHelper::InnerGenerate(FEmitterLocalContext& Context, const
 			if (!bComplete)
 			{
 				const FString LocalPathToMember = FString::Printf(TEXT("%s[%d]"), *PathToMember, Index);
+
+				// There is no point in doing diff with a "clean" struct, since we don't know what is really set by the native constructor.
 				InnerGenerate(Context, ArrayProperty->Inner, LocalPathToMember, LocalValuePtr, nullptr, true);
 			}
 		}
@@ -388,16 +375,71 @@ FString FEmitDefaultValueHelper::HandleNonNativeComponent(FEmitterLocalContext&
 
 struct FDependenciesHelper
 {
+private:
+	// Keep sync with FTypeSingletonCache::GenerateSingletonName
+	static FString GenerateZConstructor(UField* Item)
+	{
+		FString Result;
+		if (!ensure(Item))
+		{
+			return Result;
+		}
+		
+		for (UObject* Outer = Item; Outer; Outer = Outer->GetOuter())
+		{
+			if (!Result.IsEmpty())
+			{
+				Result = TEXT("_") + Result;
+			}
+
+			if (Cast<UClass>(Outer) || Cast<UScriptStruct>(Outer))
+			{
+				FString OuterName = FEmitHelper::GetCppName(CastChecked<UField>(Outer), true);
+				Result = OuterName + Result;
+
+				// Structs can also have UPackage outer.
+				if (Cast<UClass>(Outer) || Cast<UPackage>(Outer->GetOuter()))
+				{
+					break;
+				}
+			}
+			else
+			{
+				Result = Outer->GetName() + Result;
+			}
+		}
+
+		// Can't use long package names in function names.
+		if (Result.StartsWith(TEXT("/Script/"), ESearchCase::CaseSensitive))
+		{
+			Result = FPackageName::GetShortName(Result);
+		}
+
+		const FString ClassString = Item->IsA<UClass>() ? TEXT("UClass") : TEXT("UScriptStruct");
+		const FString PostFix = Item->IsA<UClass>() ? TEXT("_NoRegister") : TEXT("");
+		return FString(TEXT("Z_Construct_")) + ClassString + TEXT("_") + Result + PostFix + TEXT("()");
+	}
+public:
 	static void AddDependenciesInConstructor(FEmitterLocalContext& Context)
 	{
+		const bool bUseZConstructorInGeneratedCode = true;
 		if (Context.Dependencies.ConvertedClasses.Num())
 		{
 			Context.AddLine(TEXT("// List of all referenced converted classes"));
 		}
 		for (auto LocStruct : Context.Dependencies.ConvertedClasses)
 		{
-			Context.AddLine(FString::Printf(TEXT("CastChecked<UDynamicClass>(GetClass())->ReferencedConvertedFields.Add(%s::StaticClass());")
-				, *FEmitHelper::GetCppName(LocStruct, true)));
+			FString ClassConstructor;
+			if (bUseZConstructorInGeneratedCode)
+			{
+				ClassConstructor = GenerateZConstructor(Context.Dependencies.FindOriginalClass(LocStruct));
+				Context.AddLine(FString::Printf(TEXT("extern UClass* %s;"), *ClassConstructor));
+			}
+			else
+			{
+				ClassConstructor = FEmitHelper::GetCppName(LocStruct, true) + TEXT("::StaticClass()");
+			}
+			Context.AddLine(FString::Printf(TEXT("CastChecked<UDynamicClass>(GetClass())->ReferencedConvertedFields.Add(%s);"), *ClassConstructor));
 		}
 
 		if (Context.Dependencies.ConvertedStructs.Num())
@@ -406,8 +448,17 @@ struct FDependenciesHelper
 		}
 		for (auto LocStruct : Context.Dependencies.ConvertedStructs)
 		{
-			Context.AddLine(FString::Printf(TEXT("CastChecked<UDynamicClass>(GetClass())->ReferencedConvertedFields.Add(%s::StaticStruct());")
-				, *FEmitHelper::GetCppName(LocStruct)));
+			FString StructConstructor;
+			if (bUseZConstructorInGeneratedCode)
+			{
+				StructConstructor = GenerateZConstructor(LocStruct);
+				Context.AddLine(FString::Printf(TEXT("extern UScriptStruct* %s;"), *StructConstructor));
+			}
+			else
+			{
+				StructConstructor = FEmitHelper::GetCppName(LocStruct, true) + TEXT("::StaticStruct()");
+			}
+			Context.AddLine(FString::Printf(TEXT("CastChecked<UDynamicClass>(GetClass())->ReferencedConvertedFields.Add(%s);"), *StructConstructor));
 		}
 
 		if (Context.Dependencies.Assets.Num())
@@ -416,25 +467,48 @@ struct FDependenciesHelper
 		}
 		for (auto LocAsset : Context.Dependencies.Assets)
 		{
-			const FString AssetStr = Context.FindGloballyMappedObject(LocAsset, nullptr, true, false);
+			const FString AssetStr = Context.FindGloballyMappedObject(LocAsset, UObject::StaticClass(), true, false);
 			Context.AddLine(FString::Printf(TEXT("CastChecked<UDynamicClass>(GetClass())->UsedAssets.Add(%s);"), *AssetStr));
 		}
 	}
 
 	static void AddStaticFunctionsForDependencies(FEmitterLocalContext& Context)
 	{
-		auto OriginalClass = Context.GetCurrentlyGeneratedClass();
+		auto SourceClass = Context.GetCurrentlyGeneratedClass();
+		auto OriginalClass = Context.Dependencies.FindOriginalClass(SourceClass);
 		const FString CppClassName = FEmitHelper::GetCppName(OriginalClass);
+		
 
 		// __StaticDependenciesAssets
-		Context.AddLine(FString::Printf(TEXT("void %s::__StaticDependenciesAssets(TArray<FName>& OutPackagePaths)"), *CppClassName));
+		Context.AddLine(FString::Printf(TEXT("void %s::__StaticDependenciesAssets(TArray<FBlueprintDependencyData>& AssetsToLoad)"), *CppClassName));
 		Context.AddLine(TEXT("{"));
 		Context.IncreaseIndent();
 
 		for (UObject* LocAsset : Context.Dependencies.Assets)
 		{
-			const FString PakagePath = LocAsset->GetOutermost()->GetPathName();
-			Context.AddLine(FString::Printf(TEXT("OutPackagePaths.Add(TEXT(\"%s\"));"), *PakagePath));
+			auto GetConvertedClass = [&](const UClass* AssetType) -> const UClass*
+			{
+				const UBlueprintGeneratedClass* BPGC = Cast<const UBlueprintGeneratedClass>(AssetType);
+				if (AssetType->IsChildOf<UUserDefinedEnum>())
+				{
+					return UEnum::StaticClass();
+				}
+				else if (AssetType->IsChildOf<UUserDefinedStruct>())
+				{
+					return UScriptStruct::StaticClass();
+				}
+				else if (BPGC && Context.Dependencies.WillClassBeConverted(BPGC))
+				{
+					return UDynamicClass::StaticClass();
+				}
+				return AssetType;
+			};
+			const UClass* ClassToUse = GetConvertedClass(LocAsset->GetClass());
+			Context.AddLine(FString::Printf(TEXT("AssetsToLoad.Add({FName(TEXT(\"%s\")), FName(TEXT(\"%s\")), FName(TEXT(\"%s\")), FName(TEXT(\"%s\"))});")
+				, *LocAsset->GetOutermost()->GetPathName()
+				, *LocAsset->GetName()
+				, *ClassToUse->GetOutermost()->GetPathName()
+				, *ClassToUse->GetName()));
 		}
 
 		Context.DecreaseIndent();
@@ -446,44 +520,15 @@ struct FDependenciesHelper
 		Context.AddLine(TEXT("{"));
 		Context.IncreaseIndent();
 
-		{
-			// __StaticDependenciesAssets
-			Context.DecreaseIndent();
-			Context.AddLine(TEXT("#if WITH_EDITOR"));
-			Context.IncreaseIndent();
-			Context.AddLine(TEXT("static UField* __StaticCreateClass()"));
-			Context.AddLine(TEXT("{"));
-			Context.IncreaseIndent();
-			Context.AddLine(FString::Printf(TEXT("return CastChecked<UDynamicClass>(%s::StaticClass());"), *FEmitHelper::GetCppName(OriginalClass, true)));
-			Context.DecreaseIndent();
-			Context.AddLine(TEXT("}"));
-			Context.DecreaseIndent();
-			Context.AddLine(TEXT("#endif //WITH_EDITOR"));
-			Context.IncreaseIndent();
-		}
-
 		Context.AddLine(FString::Printf(TEXT("%s()"), *RegisterHelperName));
 		Context.AddLine(TEXT("{"));
 		Context.IncreaseIndent();
 
 		Context.AddLine(FString::Printf(
 			TEXT("FConvertedBlueprintsDependencies::Get().RegisterClass(TEXT(\"%s\"), &%s::__StaticDependenciesAssets);")
-			, *OriginalClass->GetName()
-			, *CppClassName
+			, *OriginalClass->GetOutermost()->GetPathName()
 			, *CppClassName));
 
-		{
-			Context.DecreaseIndent();
-			Context.AddLine(TEXT("#if WITH_EDITOR"));
-			Context.IncreaseIndent();
-
-			Context.AddLine(TEXT("FConvertedBlueprintsDependencies::Get().RegisterClass(&__StaticCreateClass);"));
-
-			Context.DecreaseIndent();
-			Context.AddLine(TEXT("#endif //WITH_EDITOR"));
-			Context.IncreaseIndent();
-		}
-
 		Context.DecreaseIndent();
 		Context.AddLine(TEXT("}"));
 
@@ -501,7 +546,12 @@ void FEmitDefaultValueHelper::GenerateConstructor(FEmitterLocalContext& Context)
 	auto BPGC = CastChecked<UBlueprintGeneratedClass>(Context.GetCurrentlyGeneratedClass());
 	const FString CppClassName = FEmitHelper::GetCppName(BPGC);
 
-	Context.AddLine(FString::Printf(TEXT("%s::%s(const FObjectInitializer& ObjectInitializer) : Super(ObjectInitializer)"), *CppClassName, *CppClassName));
+	UClass* SuperClass = BPGC->GetSuperClass();
+	const bool bSuperHasOnlyDefaultConstructor = SuperClass && SuperClass->HasMetaData(TEXT("OnlyDefaultConstructorDeclared"));
+	Context.AddLine(FString::Printf(TEXT("%s::%s(const FObjectInitializer& ObjectInitializer) : Super(%s)")
+		, *CppClassName
+		, *CppClassName
+		, bSuperHasOnlyDefaultConstructor ? TEXT("") : TEXT("ObjectInitializer")));
 	Context.AddLine(TEXT("{"));
 	Context.IncreaseIndent();
 
@@ -535,8 +585,11 @@ void FEmitDefaultValueHelper::GenerateConstructor(FEmitterLocalContext& Context)
 		Context.AddLine(TEXT("ensure(0 == CastChecked<UDynamicClass>(GetClass())->DynamicBindingObjects.Num());"));
 		Context.AddLine(TEXT("ensure(0 == CastChecked<UDynamicClass>(GetClass())->ComponentTemplates.Num());"));
 		Context.AddLine(TEXT("ensure(0 == CastChecked<UDynamicClass>(GetClass())->Timelines.Num());"));
+		Context.AddLine(TEXT("ensure(nullptr == CastChecked<UDynamicClass>(GetClass())->AnimClassImplementation);"));
 		Context.CurrentCodeType = FEmitterLocalContext::EGeneratedCodeType::SubobjectsOfClass;
 
+		FDependenciesHelper::AddDependenciesInConstructor(Context);
+
 		auto CreateAndInitializeClassSubobjects = [&](bool bCreate, bool bInitilize)
 		{
 			for (auto ComponentTemplate : ActorComponentTempatesOwnedByClass)
@@ -566,7 +619,8 @@ void FEmitDefaultValueHelper::GenerateConstructor(FEmitterLocalContext& Context)
 		};
 		CreateAndInitializeClassSubobjects(true, false);
 		CreateAndInitializeClassSubobjects(false, true);
-		FDependenciesHelper::AddDependenciesInConstructor(Context);
+
+		FBackendHelperAnim::CreateAnimClassData(Context);
 
 		Context.DecreaseIndent();
 		Context.AddLine(TEXT("}"));
@@ -662,12 +716,17 @@ void FEmitDefaultValueHelper::GenerateConstructor(FEmitterLocalContext& Context)
 						{
 							const FString NativeVariablePropertyName = HandleNonNativeComponent(Context, Node, HandledProperties, NativeCreatedComponentProperties, nullptr, ComponentsToInit);
 
-							if (i == 0 && bNeedsRootComponentAssignment && Node->ComponentTemplate && Node->ComponentTemplate->IsA<USceneComponent>() && !NativeVariablePropertyName.IsEmpty())
+							if (bNeedsRootComponentAssignment && Node->ComponentTemplate && Node->ComponentTemplate->IsA<USceneComponent>() && !NativeVariablePropertyName.IsEmpty())
 							{
-								Context.AddLine(FString::Printf(TEXT("RootComponent = %s;"), *NativeVariablePropertyName));
+								// Only emit the explicit root component assignment statement if we're looking at the child BPGC that we're generating ctor code
+								// for. In all other cases, the root component will already be set up by a chained parent ctor call, so we avoid stomping it here.
+								if (i == 0)
+								{
+									Context.AddLine(FString::Printf(TEXT("RootComponent = %s;"), *NativeVariablePropertyName));
+									HandledProperties.Add(RootComponentProperty);
+								}
 
 								bNeedsRootComponentAssignment = false;
-								HandledProperties.Add(RootComponentProperty);
 							}
 						}
 					}
@@ -729,11 +788,19 @@ FString FEmitDefaultValueHelper::HandleClassSubobject(FEmitterLocalContext& Cont
 	FString LocalNativeName;
 	if (bCreate)
 	{
-		const FString OuterStr = Context.FindGloballyMappedObject(Object->GetOuter());
+		FString OuterStr = Context.FindGloballyMappedObject(Object->GetOuter());
 		if (OuterStr.IsEmpty())
 		{
-			ensure(false);
-			return FString();
+			OuterStr = HandleClassSubobject(Context, Object->GetOuter(), ListOfSubobjectsType, bCreate, bInitilize);
+			if (OuterStr.IsEmpty())
+			{
+				return FString();
+			}
+			const FString AlreadyCreatedObject = Context.FindGloballyMappedObject(Object);
+			if (!AlreadyCreatedObject.IsEmpty())
+			{
+				return AlreadyCreatedObject;
+			}
 		}
 
 		const bool AddAsSubobjectOfClass = Object->GetOuter() == Context.GetCurrentlyGeneratedClass();
diff --git a/Engine/Source/Developer/BlueprintCompilerCppBackend/Public/BlueprintCompilerCppBackendGatherDependencies.h b/Engine/Source/Developer/BlueprintCompilerCppBackend/Public/BlueprintCompilerCppBackendGatherDependencies.h
--- a/Engine/Source/Developer/BlueprintCompilerCppBackend/Public/BlueprintCompilerCppBackendGatherDependencies.h
+++ b/Engine/Source/Developer/BlueprintCompilerCppBackend/Public/BlueprintCompilerCppBackendGatherDependencies.h
@@ -7,15 +7,13 @@
 #include "Engine/UserDefinedEnum.h"
 
 /** The struct gathers dependencies of a converted BPGC */
-struct BLUEPRINTCOMPILERCPPBACKEND_API FGatherConvertedClassDependencies : public FReferenceCollector
+struct BLUEPRINTCOMPILERCPPBACKEND_API FGatherConvertedClassDependencies
 {
 protected:
-	TSet<UObject*> SerializedObjects;
 	UStruct* OriginalStruct;
 
 public:
 	// Dependencies:
-
 	TArray<UObject*> Assets; 
 
 	TSet<UBlueprintGeneratedClass*> ConvertedClasses;
@@ -29,21 +27,16 @@ public:
 public:
 	FGatherConvertedClassDependencies(UStruct* InStruct);
 
-	virtual bool IsIgnoringArchetypeRef() const override { return false; }
-	virtual bool IsIgnoringTransient() const override { return true; }
-
-	virtual void HandleObjectReference(UObject*& InObject, const UObject* InReferencingObject, const UProperty* InReferencingProperty) override;
-
-	UStruct* GetOriginalStruct() const
+	UStruct* GetActualStruct() const
 	{
 		return OriginalStruct;
 	}
 
+	UClass* FindOriginalClass(const UClass* InClass) const;
+
 public:
 	bool WillClassBeConverted(const UBlueprintGeneratedClass* InClass) const;
 
 protected:
-	void FindReferences(UObject* Object);
 	void DependenciesForHeader();
-	bool ShouldIncludeHeaderFor(UObject* Object) const;
 };
diff --git a/Engine/Source/Developer/BlueprintCompilerCppBackend/Public/IBlueprintCompilerCppBackendModule.h b/Engine/Source/Developer/BlueprintCompilerCppBackend/Public/IBlueprintCompilerCppBackendModule.h
--- a/Engine/Source/Developer/BlueprintCompilerCppBackend/Public/IBlueprintCompilerCppBackendModule.h
+++ b/Engine/Source/Developer/BlueprintCompilerCppBackend/Public/IBlueprintCompilerCppBackendModule.h
@@ -29,5 +29,21 @@ public:
 	 * @return A delegate that will be used to fill out the shared include statement in generated source.
 	 */
 	virtual FPCHFilenameQuery& OnPCHFilenameQuery() = 0;
+
+
+	DECLARE_DELEGATE_RetVal_OneParam(bool, FIsTargetedForConversionQuery, const UObject*);
+	/**
+	 * Provides a hook so that external modules can specify which assets will 
+	 * and won't be converted (so the backend knows how to handle cross asset 
+	 * references).
+	 * 
+	 * @return A delegate that will used to determine how to generate cross asset references (in generated source).
+	 */
+	virtual FIsTargetedForConversionQuery& OnIsTargetedForConversionQuery() = 0;
+
+	/**
+	 *	Cpp is generated from a duplicate of the original Blueprint. This map stores info about the original BPGC.
+	 */
+	virtual TMap<TWeakObjectPtr<UClass>, TWeakObjectPtr<UClass> >& GetOriginalClassMap() = 0;
 };
 
diff --git a/Engine/Source/Developer/BlueprintNativeCodeGen/BlueprintNativeCodeGen.Build.cs b/Engine/Source/Developer/BlueprintNativeCodeGen/BlueprintNativeCodeGen.Build.cs
--- a/Engine/Source/Developer/BlueprintNativeCodeGen/BlueprintNativeCodeGen.Build.cs
+++ b/Engine/Source/Developer/BlueprintNativeCodeGen/BlueprintNativeCodeGen.Build.cs
@@ -30,6 +30,7 @@ public class BlueprintNativeCodeGen : ModuleRules
                 "BlueprintCompilerCppBackend",
                 "GameProjectGeneration",
                 "Projects",
+                "Kismet",
 			}
 		);
 
diff --git a/Engine/Source/Developer/BlueprintNativeCodeGen/Private/BlueprintNativeCodeGenManifest.cpp b/Engine/Source/Developer/BlueprintNativeCodeGen/Private/BlueprintNativeCodeGenManifest.cpp
--- a/Engine/Source/Developer/BlueprintNativeCodeGen/Private/BlueprintNativeCodeGenManifest.cpp
+++ b/Engine/Source/Developer/BlueprintNativeCodeGen/Private/BlueprintNativeCodeGenManifest.cpp
@@ -1,14 +1,17 @@
 // Copyright 1998-2015 Epic Games, Inc. All Rights Reserved.
 
 #include "BlueprintNativeCodeGenPCH.h"
+
+#include "App.h"								// for GetGameName()
 #include "BlueprintNativeCodeGenManifest.h"
+#include "BlueprintNativeCodeGenModule.h"
+#include "IBlueprintCompilerCppBackendModule.h" // for ConstructBaseFilename()
+#include "Engine/Blueprint.h"					// for GeneratedClass
+#include "JsonObjectConverter.h"
 #include "NativeCodeGenCommandlineParams.h"
-#include "App.h" // for GetGameName()
 #include "Serialization/JsonReader.h"
 #include "Serialization/JsonSerializer.h"
-#include "JsonObjectConverter.h"
 #include "Serialization/JsonWriter.h"
-#include "IBlueprintCompilerCppBackendModule.h" // for ConstructBaseFilename()
 
 DEFINE_LOG_CATEGORY_STATIC(LogNativeCodeGenManifest, Log, All);
 
@@ -18,17 +21,18 @@ DEFINE_LOG_CATEGORY_STATIC(LogNativeCodeGenManifest, Log, All);
 
 namespace BlueprintNativeCodeGenManifestImpl
 {
-	static const int64 CPF_NoFlags = 0x00;
-	static const FString ManifestFileExt    = TEXT(".BpCodeGenManifest.json");
-	static const FString CppFileExt	        = TEXT(".cpp");
-	static const FString HeaderFileExt      = TEXT(".h");
-	static const FString HeaderSubDir       = TEXT("Public");
-	static const FString CppSubDir          = TEXT("Private");
-	static const FString ModuleBuildFileExt = TEXT(".Build.cs");
-	static const FString FallbackPluginName = TEXT("GeneratedBpCode");
-	static const FString PreviewFilePostfix = TEXT("-Preview");
-	static const FString PluginFileExt      = TEXT(".uplugin");
-	static const FString SourceSubDir       = TEXT("Source");
+	static const int64   CPF_NoFlags          = 0x00;
+	static const FString ManifestFileExt      = TEXT(".BpCodeGenManifest.json");
+	static const FString CppFileExt	          = TEXT(".cpp");
+	static const FString HeaderFileExt        = TEXT(".h");
+	static const FString HeaderSubDir         = TEXT("Public");
+	static const FString CppSubDir            = TEXT("Private");
+	static const FString ModuleBuildFileExt   = TEXT(".Build.cs");
+	static const FString FallbackPluginName   = TEXT("NativizedScript");
+	static const FString PreviewFilePostfix   = TEXT("-Preview");
+	static const FString PluginFileExt        = TEXT(".uplugin");
+	static const FString SourceSubDir         = TEXT("Source");
+	static const FString EditorModulePostfix  = TEXT("Editor");
 
 	/**
 	 * Populates the provided manifest object with data from the specified file.
@@ -40,7 +44,7 @@ namespace BlueprintNativeCodeGenManifestImpl
 	static bool LoadManifest(const FString& FilePath, FBlueprintNativeCodeGenManifest* Manifest);
 
 	/**
-	 * Helper function the homogenizes file/directory paths so that they can be 
+	 * Helper function that homogenizes file/directory paths so that they can be 
 	 * compared for equivalence against others.
 	 * 
 	 * @param  DirectoryPath    The path that you want sanitized.
@@ -49,41 +53,41 @@ namespace BlueprintNativeCodeGenManifestImpl
 	static FString GetComparibleDirPath(const FString& DirectoryPath);
 
 	/**
-	 * Constructs a directory path for the plugin's source files.
+	 * Retrieves the sub-directory for either header or cpp source files 
+	 * (depending on which was requested).
 	 * 
-	 * @param  PluginPath    A path to the root directory for the plugin.
-	 * @return A path to the plugin's source directory.
+	 * @param  SourceType    Defines the type of source file to return for (header or cpp).
+	 * @return A directory name for the specified source file type.
 	 */
-	static FString GetTargetSourceDir(const FString& PluginPath);
+	static const FString& GetSourceSubDir(const FBlueprintNativeCodeGenPaths::ESourceFileType SourceType);
 
 	/**
-	 * Constructs a path for the specified module, within the target plugin.
+	 * Retrieves the extension for either header or cpp source files (depending
+	 * on which was requested).
 	 * 
-	 * @param  PluginPath    A path to the root directory for the plugin.
-	 * @param  ModuleName    The name of the module that you're generating for the plugin.
-	 * @return A path to the module's directory, within the plugin.
+	 * @param  SourceType    Defines the type of source file to return for (header or cpp).
+	 * @return A file extension (including the leading dot), for the specified source file type.
 	 */
-	static FString GetModuleDirectory(const FString& PluginPath, const FString& ModuleName);
+	static const FString& GetSourceFileExt(const FBlueprintNativeCodeGenPaths::ESourceFileType SourceType);
 
 	/**
-	 * Constructs a header path for the specified asset.
+	 * Constructs a source file path for the specified asset.
 	 * 
-	 * @param  PluginPath	A path to the root directory for the plugin.
+	 * @param  TargetPaths	Specified the destination directory for the file.
 	 * @param  Asset		The asset you want a header file for.
-	 * @param  ModuleName	The name of the module that you're generating for the plugin.
+	 * @param  SourceType	Defines the type of source file to generate for (header or cpp).
 	 * @return A target file path for the specified asset to save a header to.
 	 */
-	static FString GenerateHeaderSavePath(const FString& PluginPath, const FAssetData& Asset, const FString& ModuleName);
+	static FString GenerateSourceFileSavePath(const FBlueprintNativeCodeGenPaths& TargetPaths, const FAssetData& Asset, const FBlueprintNativeCodeGenPaths::ESourceFileType SourceType);
 
 	/**
-	 * Constructs a cpp path for the specified asset.
 	 * 
-	 * @param  PluginPath	A path to the root directory for the plugin.
-	 * @param  Asset		The asset you want a cpp file for.
-	 * @param  ModuleName	The name of the module that you're generating for the plugin.
-	 * @return A target file path for the specified asset to save a cpp file to.
+	 * 
+	 * @param  TargetPaths    
+	 * @param  Asset    
+	 * @return 
 	 */
-	static FString GenerateCppSavePath(const FString& PluginPath, const FAssetData& Asset, const FString& ModuleName);
+	static FString GenerateUnconvertedWrapperPath(const FBlueprintNativeCodeGenPaths& TargetPaths, const FAssetData& Asset);
 
 	/**
 	 * Coordinates with the code-gen backend, to produce a base filename (one 
@@ -103,6 +107,36 @@ namespace BlueprintNativeCodeGenManifestImpl
 	 * @return False if the function failed to collect dependencies for the specified object.
 	 */
 	static bool GatherModuleDependencies(const UObject* AssetObj, TArray<UPackage*>& DependenciesOut);
+
+	/**
+	 * Obtains the reflected name for the native field (class/enum/struct) that 
+	 * we'll generate to replace the specified asset.
+	 *
+	 * @param  Asset	The asset you want a name from.
+	 * @return The name of the asset field (class/enum/struct).
+	 */
+	static FString GetFieldName(const FAssetData& Asset);
+
+	/**
+	 * The object returned by FAssetData::GetAsset() doesn't always give us the 
+	 * target object that will be replaced (for Blueprint's, it would be the 
+	 * class instead). So this helper function will suss out the right object 
+	 * for you.
+	 * 
+	 * @param  Asset    The asset you want an object for.
+	 * @return A pointer to the targeted object from the asset's package.
+	 */
+	static UField* GetTargetAssetObject(const FAssetData& Asset);
+
+	/**
+	 * Returns the object path for the field from the specified asset's package
+	 * that is being replaced (Asset.ObjectPath will not suffice, as that 
+	 * does not always reflect the object that is being replaced).
+	 * 
+	 * @param  Asset    The asset you want an object-path for.
+	 * @return An object-path for the target field-object within the asset's package.
+	 */
+	static FString GetTargetObjectPath(const FAssetData& Asset);
 }
 
 //------------------------------------------------------------------------------
@@ -124,27 +158,28 @@ static bool BlueprintNativeCodeGenManifestImpl::LoadManifest(const FString& File
 }
 
 //------------------------------------------------------------------------------
-static FString BlueprintNativeCodeGenManifestImpl::GetTargetSourceDir(const FString& PluginPath)
+static const FString& BlueprintNativeCodeGenManifestImpl::GetSourceSubDir(const FBlueprintNativeCodeGenPaths::ESourceFileType SourceType)
 {
-	return FPaths::Combine(*PluginPath, *SourceSubDir);
+	return (SourceType == FBlueprintNativeCodeGenPaths::HFile) ? HeaderSubDir : CppSubDir;
 }
 
 //------------------------------------------------------------------------------
-static FString BlueprintNativeCodeGenManifestImpl::GetModuleDirectory(const FString& PluginPath, const FString& ModuleName)
+static const FString& BlueprintNativeCodeGenManifestImpl::GetSourceFileExt(const FBlueprintNativeCodeGenPaths::ESourceFileType SourceType)
 {
-	return FPaths::Combine(*GetTargetSourceDir(PluginPath), *ModuleName);
+	return (SourceType == FBlueprintNativeCodeGenPaths::HFile) ? HeaderFileExt : CppFileExt;
 }
 
 //------------------------------------------------------------------------------
-static FString BlueprintNativeCodeGenManifestImpl::GenerateHeaderSavePath(const FString& PluginPath, const FAssetData& Asset, const FString& ModuleName)
+static FString BlueprintNativeCodeGenManifestImpl::GenerateSourceFileSavePath(const FBlueprintNativeCodeGenPaths& TargetPaths, const FAssetData& Asset, const FBlueprintNativeCodeGenPaths::ESourceFileType SourceType)
 {
-	return FPaths::Combine( *FPaths::Combine(*GetModuleDirectory(PluginPath, ModuleName), *HeaderSubDir), *GetBaseFilename(Asset) ) + HeaderFileExt;
+	return FPaths::Combine(*TargetPaths.RuntimeSourceDir(SourceType), *GetBaseFilename(Asset)) + GetSourceFileExt(SourceType);
 }
 
 //------------------------------------------------------------------------------
-static FString BlueprintNativeCodeGenManifestImpl::GenerateCppSavePath(const FString& PluginPath, const FAssetData& Asset, const FString& ModuleName)
+static FString BlueprintNativeCodeGenManifestImpl::GenerateUnconvertedWrapperPath(const FBlueprintNativeCodeGenPaths& TargetPaths, const FAssetData& Asset)
 {
-	return FPaths::Combine( *FPaths::Combine(*GetModuleDirectory(PluginPath, ModuleName), *CppSubDir), *GetBaseFilename(Asset) ) + CppFileExt;
+	const FBlueprintNativeCodeGenPaths::ESourceFileType WrapperFileType = FBlueprintNativeCodeGenPaths::HFile;
+	return FPaths::Combine(*TargetPaths.RuntimeSourceDir(WrapperFileType), *GetBaseFilename(Asset)) + GetSourceFileExt(WrapperFileType);
 }
 
 //------------------------------------------------------------------------------
@@ -201,7 +236,7 @@ static bool BlueprintNativeCodeGenManifestImpl::GatherModuleDependencies(const U
 			}
 
 			// we want only native packages, ones that are not editor-only
-			if ((DependentPackage->GetPackageFlags() & (PKG_CompiledIn | PKG_EditorOnly)) == PKG_CompiledIn)
+			if ((DependentPackage->GetPackageFlags() & (PKG_CompiledIn | PKG_EditorOnly | PKG_Developer)) == PKG_CompiledIn)
 			{
 				DependenciesOut.AddUnique(DependentPackage);// PkgImport.ObjectName.ToString());
 			}
@@ -211,47 +246,174 @@ static bool BlueprintNativeCodeGenManifestImpl::GatherModuleDependencies(const U
 	return bFoundLinker;
 }
 
+//------------------------------------------------------------------------------
+static FString BlueprintNativeCodeGenManifestImpl::GetFieldName(const FAssetData& Asset)
+{
+	UField* AssetField = GetTargetAssetObject(Asset);
+	return (AssetField != nullptr) ? AssetField->GetName() : TEXT("");
+}
+
+//------------------------------------------------------------------------------
+static UField* BlueprintNativeCodeGenManifestImpl::GetTargetAssetObject(const FAssetData& Asset)
+{
+	UObject* AssetObj = Asset.GetAsset();
+
+	UField* AssetField = nullptr;
+	if (UBlueprint* BlueprintAsset = Cast<UBlueprint>(AssetObj))
+	{
+		AssetField = BlueprintAsset->GeneratedClass;
+		if (!AssetField)
+		{
+			UE_LOG(LogNativeCodeGenManifest, Warning, TEXT("null BPGC in %s"), *BlueprintAsset->GetPathName());
+		}
+	}
+	else
+	{
+		// only other asset types that we should be converting are enums and 
+		// structs (both UFields)
+		AssetField = CastChecked<UField>(AssetObj);
+	}
+
+	return AssetField;
+}
+
+//------------------------------------------------------------------------------
+static FString BlueprintNativeCodeGenManifestImpl::GetTargetObjectPath(const FAssetData& Asset)
+{
+	UField* AssetField = GetTargetAssetObject(Asset);
+	return (AssetField != nullptr) ? AssetField->GetPathName() : TEXT("");
+}
+
 /*******************************************************************************
  * FConvertedAssetRecord
  ******************************************************************************/
  
 //------------------------------------------------------------------------------
-FConvertedAssetRecord::FConvertedAssetRecord(const FAssetData& AssetInfo, const FString& TargetDir, const FString& ModuleName)
-	: AssetPtr(AssetInfo.GetAsset())
-	, AssetType(AssetInfo.GetClass())
-	, AssetPath(AssetInfo.PackageName.ToString())
-	, GeneratedHeaderPath(BlueprintNativeCodeGenManifestImpl::GenerateHeaderSavePath(TargetDir, AssetInfo, ModuleName))
-	, GeneratedCppPath(BlueprintNativeCodeGenManifestImpl::GenerateCppSavePath(TargetDir, AssetInfo, ModuleName))
+FConvertedAssetRecord::FConvertedAssetRecord(const FAssetData& AssetInfo, const FBlueprintNativeCodeGenPaths& TargetPaths)
+	: AssetType(AssetInfo.GetClass())
+	, TargetObjPath(BlueprintNativeCodeGenManifestImpl::GetTargetObjectPath(AssetInfo))
 {
+	GeneratedCppPath    = BlueprintNativeCodeGenManifestImpl::GenerateSourceFileSavePath(TargetPaths, AssetInfo, FBlueprintNativeCodeGenPaths::CppFile);
+	GeneratedHeaderPath = BlueprintNativeCodeGenManifestImpl::GenerateSourceFileSavePath(TargetPaths, AssetInfo, FBlueprintNativeCodeGenPaths::HFile);
 }
 
 //------------------------------------------------------------------------------
 bool FConvertedAssetRecord::IsValid()
 {
 	// every conversion will have a header file (interfaces don't have implementation files)
-	return AssetPtr.IsValid() && !GeneratedHeaderPath.IsEmpty() && (AssetType != nullptr) && !AssetPath.IsEmpty();
+	return ToAssetRef().IsValid() && (AssetType != nullptr) && !GeneratedHeaderPath.IsEmpty();
+}
+
+//------------------------------------------------------------------------------
+FStringAssetReference FConvertedAssetRecord::ToAssetRef() const
+{
+	return FStringAssetReference(TargetObjPath);
 }
 
 /*******************************************************************************
- * FBlueprintNativeCodeGenManifest
+ * FUnconvertedDependencyRecord
  ******************************************************************************/
+ 
+//------------------------------------------------------------------------------
+FUnconvertedDependencyRecord::FUnconvertedDependencyRecord(const FAssetData& AssetInfo, const FBlueprintNativeCodeGenPaths& TargetPaths)
+	: GeneratedWrapperPath(BlueprintNativeCodeGenManifestImpl::GenerateUnconvertedWrapperPath(TargetPaths, AssetInfo))
+{
+}
 
 //------------------------------------------------------------------------------
-FString FBlueprintNativeCodeGenManifest::GetDefaultFilename()
+bool FUnconvertedDependencyRecord::IsValid()
+{
+	return !GeneratedWrapperPath.IsEmpty();
+}
+
+/*******************************************************************************
+ * FBlueprintNativeCodeGenPaths
+ ******************************************************************************/
+
+//------------------------------------------------------------------------------
+FString FBlueprintNativeCodeGenPaths::GetDefaultManifestPath()
 {
 	return FApp::GetGameName() + BlueprintNativeCodeGenManifestImpl::ManifestFileExt;
 }
 
 //------------------------------------------------------------------------------
-FBlueprintNativeCodeGenManifest::FBlueprintNativeCodeGenManifest(const FString TargetPath)
-	: OutputDir(TargetPath)
+FBlueprintNativeCodeGenPaths::FBlueprintNativeCodeGenPaths(const FString& PluginNameIn, const FString& TargetDirIn, const FString& ManifestPathIn)
+	: TargetDir(TargetDirIn)
+	, PluginName(PluginNameIn)
+	, ManifestPath(ManifestPathIn)
 {
-	if (OutputDir.IsEmpty())
-	{
-		OutputDir = FPaths::GameIntermediateDir();
-	}
-	// do NOT load from an existing interface, as this is the default 
-	// constructor used by the USTRUCT() system
+}
+
+//------------------------------------------------------------------------------
+FString FBlueprintNativeCodeGenPaths::ManifestFilePath() const
+{
+	return ManifestPath;
+}
+
+//------------------------------------------------------------------------------
+FString FBlueprintNativeCodeGenPaths::PluginRootDir() const
+{
+	return TargetDir;
+}
+
+//------------------------------------------------------------------------------
+FString FBlueprintNativeCodeGenPaths::PluginFilePath() const
+{
+	return FPaths::Combine(*PluginRootDir(), *PluginName) + BlueprintNativeCodeGenManifestImpl::PluginFileExt;;
+}
+
+//------------------------------------------------------------------------------
+FString FBlueprintNativeCodeGenPaths::PluginSourceDir() const
+{
+	return FPaths::Combine(*PluginRootDir(), *BlueprintNativeCodeGenManifestImpl::SourceSubDir);
+}
+
+//------------------------------------------------------------------------------
+FString FBlueprintNativeCodeGenPaths::RuntimeModuleDir() const
+{
+	return FPaths::Combine(*PluginSourceDir(), *RuntimeModuleName());
+}
+
+//------------------------------------------------------------------------------
+FString FBlueprintNativeCodeGenPaths::RuntimeModuleName() const
+{
+	return PluginName;
+}
+
+//------------------------------------------------------------------------------
+FString FBlueprintNativeCodeGenPaths::RuntimeBuildFile() const
+{
+	return FPaths::Combine(*RuntimeModuleDir(), *RuntimeModuleName()) + BlueprintNativeCodeGenManifestImpl::ModuleBuildFileExt;
+}
+
+//------------------------------------------------------------------------------
+FString FBlueprintNativeCodeGenPaths::RuntimeSourceDir(ESourceFileType SourceType) const
+{
+	return FPaths::Combine(*RuntimeModuleDir(), *BlueprintNativeCodeGenManifestImpl::GetSourceSubDir(SourceType));
+}
+
+//------------------------------------------------------------------------------
+FString FBlueprintNativeCodeGenPaths::RuntimeModuleFile(ESourceFileType SourceType) const
+{
+	return FPaths::Combine(*RuntimeSourceDir(SourceType), *RuntimeModuleName()) + BlueprintNativeCodeGenManifestImpl::GetSourceFileExt(SourceType);
+}
+
+//------------------------------------------------------------------------------
+FString FBlueprintNativeCodeGenPaths::RuntimePCHFilename() const
+{
+	return FPaths::GetCleanFilename(RuntimeModuleFile(HFile));
+}
+
+/*******************************************************************************
+ * FBlueprintNativeCodeGenManifest
+ ******************************************************************************/
+
+//------------------------------------------------------------------------------
+FBlueprintNativeCodeGenManifest::FBlueprintNativeCodeGenManifest()
+	: OutputDir(FPaths::ConvertRelativePathToFull( FPaths::Combine(*FPaths::Combine(*FPaths::GameIntermediateDir(), TEXT("Plugins")), *BlueprintNativeCodeGenManifestImpl::FallbackPluginName) ))
+{
+	PluginName = FPaths::GetBaseFilename(OutputDir);
+	ManifestFilePath = FPaths::Combine(*GetTargetDir(), *FBlueprintNativeCodeGenPaths::GetDefaultManifestPath());
 }
 
 //------------------------------------------------------------------------------
@@ -282,7 +444,7 @@ FBlueprintNativeCodeGenManifest::FBlueprintNativeCodeGenManifest(const FNativeCo
 			OutputDir = FPaths::Combine(*FPaths::GameIntermediateDir(), *PluginName);
 		}
 
-		FString ManifestFilename = GetDefaultFilename();
+		FString ManifestFilename = FBlueprintNativeCodeGenPaths::GetDefaultManifestPath();
 		if (CommandlineParams.bPreviewRequested)
 		{
 			ManifestFilename.InsertAt(ManifestFilename.Find(ManifestFileExt, ESearchCase::CaseSensitive, ESearchDir::FromEnd), PreviewFilePostfix);
@@ -321,7 +483,7 @@ FBlueprintNativeCodeGenManifest::FBlueprintNativeCodeGenManifest(const FNativeCo
 		if (bNewPluginNameRequested && !CommandlineParams.bPreviewRequested)
 		{
 			// delete the old plugin file (if one exists)
-			IFileManager::Get().Delete(*GetPluginFilePath());
+			IFileManager::Get().Delete(*GetTargetPaths().PluginFilePath());
 		}
 
 		// if we were only interested in obtaining the plugin path
@@ -338,75 +500,54 @@ FBlueprintNativeCodeGenManifest::FBlueprintNativeCodeGenManifest(const FNativeCo
 					*CommandlineParams.PluginName, *PluginName);
 				PluginName = CommandlineParams.PluginName;
 			}
-			ConstructRecordLookupTable();
 		}
 	}
 }
 
 //------------------------------------------------------------------------------
-TArray<FString> FBlueprintNativeCodeGenManifest::GetDestinationPaths() const
+FBlueprintNativeCodeGenManifest::FBlueprintNativeCodeGenManifest(const FString& ManifestFilePathIn)
+	: ManifestFilePath(ManifestFilePathIn)
 {
-	const FString TargetPath = GetTargetDir();
-
-	TArray<FString> DestPaths;
-	DestPaths.Add(ManifestFilePath);
-	DestPaths.Add(BlueprintNativeCodeGenManifestImpl::GetTargetSourceDir(TargetPath));
-	DestPaths.Add(GetPluginFilePath());
-
-	return DestPaths;
+	BlueprintNativeCodeGenManifestImpl::LoadManifest(ManifestFilePathIn, this);
 }
 
 //------------------------------------------------------------------------------
-FString FBlueprintNativeCodeGenManifest::GetPluginFilePath() const
+FBlueprintNativeCodeGenPaths FBlueprintNativeCodeGenManifest::GetTargetPaths() const
 {
-	return FPaths::Combine(*GetTargetDir(), *PluginName) + BlueprintNativeCodeGenManifestImpl::PluginFileExt;
+	return FBlueprintNativeCodeGenPaths(PluginName, GetTargetDir(), ManifestFilePath);
 }
 
 //------------------------------------------------------------------------------
-FString FBlueprintNativeCodeGenManifest::GetModuleFilePath() const
+FConvertedAssetRecord& FBlueprintNativeCodeGenManifest::CreateConversionRecord(const FAssetId Key, const FAssetData& AssetInfo)
 {
-	return FPaths::Combine( *BlueprintNativeCodeGenManifestImpl::GetModuleDirectory(GetTargetDir(), GetTargetModuleName()), *GetTargetModuleName() ) + 
-		BlueprintNativeCodeGenManifestImpl::ModuleBuildFileExt;
-}
+	const FBlueprintNativeCodeGenPaths TargetPaths = GetTargetPaths();
 
-//------------------------------------------------------------------------------
-FConvertedAssetRecord& FBlueprintNativeCodeGenManifest::CreateConversionRecord(const FAssetData& AssetInfo)
-{
-	const FString AssetPath  = AssetInfo.PackageName.ToString();
-	const FString TargetPath = GetTargetDir();
+	UClass* AssetType = AssetInfo.GetClass();
 	// load the asset (if it isn't already)
-	const UObject* AssetObj   = AssetInfo.GetAsset();
-	const FString& ModuleName = GetTargetModuleName();
+	const UObject* AssetObj = AssetInfo.GetAsset();
 
-	FConvertedAssetRecord* ConversionRecordPtr = FindConversionRecord(AssetPath);
+	FConvertedAssetRecord* ConversionRecordPtr = ConvertedAssets.Find(Key);
 	if (ConversionRecordPtr == nullptr)
 	{
-		RecordLookupTable.Add(AssetPath, ConvertedAssets.Num());
-		ConversionRecordPtr = &ConvertedAssets[ ConvertedAssets.Add(FConvertedAssetRecord(AssetInfo, TargetPath, ModuleName)) ];
+		UnconvertedDependencies.Remove(Key);
+		ConversionRecordPtr = &ConvertedAssets.Add(Key, FConvertedAssetRecord(AssetInfo, TargetPaths));
 	}
-	else if ( !ensure((ConversionRecordPtr->AssetPath == AssetPath) && (ConversionRecordPtr->AssetType == AssetInfo.GetClass())) )
+	else if (!ensure(AssetType == ConversionRecordPtr->AssetType))
 	{
-		UE_LOG(LogNativeCodeGenManifest, Warning, 
-			TEXT("The existing manifest entery for '%s' doesn't match what was expected (asset path and/or type). Updating it to match the asset."),
-			*AssetPath);
+		UE_LOG(LogNativeCodeGenManifest, Warning, TEXT("Existing manifest entry for '%s' has a different type, overriding with: '%s'"),
+			*AssetInfo.ObjectPath.ToString(), *AssetType->GetName());
 
-		ConversionRecordPtr->AssetPath = AssetPath;
-		ConversionRecordPtr->AssetType = AssetInfo.GetClass();
+		ConversionRecordPtr->AssetType = AssetType;
+		ConversionRecordPtr->TargetObjPath = BlueprintNativeCodeGenManifestImpl::GetTargetObjectPath(AssetInfo);
 	}
 	
 	FConvertedAssetRecord& ConversionRecord = *ConversionRecordPtr;
 	if (!ConversionRecord.IsValid())
 	{
-		// if this was a manifest entry that was loaded from an existing file, 
-		// then it wouldn't have the asset pointer (which is transient data)
-		if (!ConversionRecord.AssetPtr.IsValid())
-		{
-			ConversionRecord.AssetPtr = AssetObj;
-		}
 		if (ConversionRecord.GeneratedHeaderPath.IsEmpty())
 		{
-			ConversionRecord.GeneratedHeaderPath = BlueprintNativeCodeGenManifestImpl::GenerateHeaderSavePath(TargetPath, AssetInfo, ModuleName);
-			ConversionRecord.GeneratedCppPath    = BlueprintNativeCodeGenManifestImpl::GenerateCppSavePath(TargetPath, AssetInfo, ModuleName);
+			ConversionRecord.GeneratedCppPath    = BlueprintNativeCodeGenManifestImpl::GenerateSourceFileSavePath(TargetPaths, AssetInfo, FBlueprintNativeCodeGenPaths::CppFile);
+			ConversionRecord.GeneratedHeaderPath = BlueprintNativeCodeGenManifestImpl::GenerateSourceFileSavePath(TargetPaths, AssetInfo, FBlueprintNativeCodeGenPaths::HFile);
 		}
 		ensure(ConversionRecord.IsValid());
 	}
@@ -415,57 +556,50 @@ FConvertedAssetRecord& FBlueprintNativeCodeGenManifest::CreateConversionRecord(c
 }
 
 //------------------------------------------------------------------------------
-FConvertedAssetRecord* FBlueprintNativeCodeGenManifest::FindConversionRecord(const FString& AssetPath, bool bSlowSearch)
+FUnconvertedDependencyRecord& FBlueprintNativeCodeGenManifest::CreateUnconvertedDependencyRecord(const FAssetId UnconvertedAssetKey, const FAssetData& AssetInfo, const FAssetId ReferencingAsset)
 {
-	FConvertedAssetRecord* FoundRecord = nullptr;
-	if (int32* IndexPtr = RecordLookupTable.Find(AssetPath))
-	{
-		check(*IndexPtr < ConvertedAssets.Num());
-		FoundRecord = &ConvertedAssets[*IndexPtr];
-		ensure(FoundRecord->AssetPath == AssetPath);
-	}
-	else if (bSlowSearch)
+	FUnconvertedDependencyRecord* RecordPtr = UnconvertedDependencies.Find(UnconvertedAssetKey);
+	if (RecordPtr == nullptr)
 	{
-		for (FConvertedAssetRecord& ConversionRecord : ConvertedAssets)
-		{
-			if (ConversionRecord.AssetPath == AssetPath)
-			{
-				FoundRecord = &ConversionRecord;
-			}
-		}
+		const FBlueprintNativeCodeGenPaths TargetPaths = GetTargetPaths();
+		RecordPtr = &UnconvertedDependencies.Add(UnconvertedAssetKey, FUnconvertedDependencyRecord(AssetInfo, TargetPaths));
 	}
-	else
+	check(RecordPtr != nullptr);
+
+	FUnconvertedDependencyRecord& DependencyRecord = *RecordPtr;
+	if (!DependencyRecord.IsValid())
 	{
-		check(RecordLookupTable.Num() == ConvertedAssets.Num());
+		const FBlueprintNativeCodeGenPaths TargetPaths = GetTargetPaths();
+		DependencyRecord.GeneratedWrapperPath = BlueprintNativeCodeGenManifestImpl::GenerateUnconvertedWrapperPath(TargetPaths, AssetInfo);
+
+		ensure(DependencyRecord.IsValid());
 	}
-	return FoundRecord;
+
+	DependencyRecord.ReferencedBy.AddUnique(ReferencingAsset);
+	return DependencyRecord;
 }
 
 //------------------------------------------------------------------------------
-bool FBlueprintNativeCodeGenManifest::LogDependencies(const FAssetData& AssetInfo)
+const FConvertedAssetRecord* FBlueprintNativeCodeGenManifest::FindConversionRecord(const FAssetId AssetId) const
 {
-	// load the asset (if it isn't already)
-	const UObject* AssetObj = AssetInfo.GetAsset();
+	return ConvertedAssets.Find(AssetId);
+}
 
-	return BlueprintNativeCodeGenManifestImpl::GatherModuleDependencies(AssetObj, ModuleDependencies);
+//------------------------------------------------------------------------------
+void FBlueprintNativeCodeGenManifest::GatherModuleDependencies(UPackage* Package)
+{
+	BlueprintNativeCodeGenManifestImpl::GatherModuleDependencies(Package, ModuleDependencies);
 }
 
 //------------------------------------------------------------------------------
-bool FBlueprintNativeCodeGenManifest::Save(bool bSort) const
+bool FBlueprintNativeCodeGenManifest::Save() const
 {
-	if (bSort)
-	{
-		ConvertedAssets.Sort([](const FConvertedAssetRecord& Lhs, const FConvertedAssetRecord& Rhs)->bool
-		{ 
-			if (Lhs.AssetType == Rhs.AssetType)
-			{
-				return Lhs.AssetPath < Rhs.AssetPath;
-			}
-			return Lhs.AssetType->GetName() < Rhs.AssetType->GetName();
-		});
-		RecordLookupTable.Empty();
-	}
+	return SaveAs(*ManifestFilePath);
+}
 
+//------------------------------------------------------------------------------
+bool FBlueprintNativeCodeGenManifest::SaveAs(const TCHAR* Filename) const
+{
 	TSharedRef<FJsonObject> JsonObject = MakeShareable(new FJsonObject());
 
 	if (FJsonObjectConverter::UStructToJsonObject(FBlueprintNativeCodeGenManifest::StaticStruct(), this, JsonObject,
@@ -477,12 +611,31 @@ bool FBlueprintNativeCodeGenManifest::Save(bool bSort) const
 		if (FJsonSerializer::Serialize(JsonObject, JsonWriter))
 		{
 			JsonWriter->Close();
-			return FFileHelper::SaveStringToFile(FileContents, *ManifestFilePath);
+			return FFileHelper::SaveStringToFile(FileContents, Filename);
 		}
 	}
 	return false;
 }
 
+void FBlueprintNativeCodeGenManifest::Merge(const TCHAR* Filename)
+{
+	FBlueprintNativeCodeGenManifest OtherManifest = FBlueprintNativeCodeGenManifest(FString(Filename));
+	for (auto& Entry : OtherManifest.ModuleDependencies)
+	{
+		ModuleDependencies.AddUnique(Entry);
+	}
+
+	for (auto& Entry : OtherManifest.ConvertedAssets)
+	{
+		ConvertedAssets.Add(Entry.Key, Entry.Value);
+	}
+
+	for (auto& Entry : UnconvertedDependencies)
+	{
+		OtherManifest.UnconvertedDependencies.Add(Entry.Key, Entry.Value);
+	}
+}
+
 //------------------------------------------------------------------------------
 FString FBlueprintNativeCodeGenManifest::GetTargetDir() const
 {
@@ -495,22 +648,8 @@ FString FBlueprintNativeCodeGenManifest::GetTargetDir() const
 	return TargetPath;
 }
 
-
 //------------------------------------------------------------------------------
 void FBlueprintNativeCodeGenManifest::Clear()
 {
 	ConvertedAssets.Empty();
-	RecordLookupTable.Empty();
-}
-
-//------------------------------------------------------------------------------
-void FBlueprintNativeCodeGenManifest::ConstructRecordLookupTable()
-{
-	RecordLookupTable.Empty(ConvertedAssets.Num());
-
-	for (int32 RecordIndex = 0; RecordIndex < ConvertedAssets.Num(); ++RecordIndex)
-	{
-		const FConvertedAssetRecord& AssetRecord = ConvertedAssets[RecordIndex];
-		RecordLookupTable.Add(AssetRecord.AssetPath, RecordIndex);
-	}
 }
diff --git a/Engine/Source/Developer/BlueprintNativeCodeGen/Private/BlueprintNativeCodeGenManifest.h b/Engine/Source/Developer/BlueprintNativeCodeGen/Private/BlueprintNativeCodeGenManifest.h
--- a/Engine/Source/Developer/BlueprintNativeCodeGen/Private/BlueprintNativeCodeGenManifest.h
+++ b/Engine/Source/Developer/BlueprintNativeCodeGen/Private/BlueprintNativeCodeGenManifest.h
@@ -7,6 +7,7 @@
 // Forward declarations
 struct FNativeCodeGenCommandlineParams;
 class  FAssetData;
+struct FBlueprintNativeCodeGenPaths;
 
 /*******************************************************************************
  * FCodeGenAssetRecord
@@ -19,21 +20,26 @@ struct FConvertedAssetRecord
 
 public:
 	FConvertedAssetRecord() {}
-	FConvertedAssetRecord(const FAssetData& AssetInfo, const FString& TargetDir, const FString& ModuleName);
+	FConvertedAssetRecord(const FAssetData& AssetInfo, const FBlueprintNativeCodeGenPaths& TargetPaths);
 
 	/**
 	 * @return True if this record demonstates a successful conversion; false if it is incomplete 
 	 */
 	bool IsValid();
 
-public:
-	TWeakObjectPtr<UObject> AssetPtr;
+	/**
+	 * @return 
+	 */
+	FStringAssetReference ToAssetRef() const;
 
+public:
 	UPROPERTY()
 	UClass* AssetType;
 
+	// cannot use a FStringAssetReference, as the json serializer has problems 
+	// with some asset paths (for example, I had a folder named 'Folder()')
 	UPROPERTY()
-	FString AssetPath;
+	FString TargetObjPath; 
 
 	UPROPERTY()
 	FString GeneratedHeaderPath;
@@ -42,6 +48,71 @@ public:
 	FString GeneratedCppPath;
 };
 
+/*******************************************************************************
+ * FUnconvertedDependencyRecord
+ ******************************************************************************/
+
+USTRUCT()
+struct FUnconvertedDependencyRecord
+{
+	GENERATED_USTRUCT_BODY()
+
+public:
+	FUnconvertedDependencyRecord() {}
+	FUnconvertedDependencyRecord(const FAssetData& AssetInfo, const FBlueprintNativeCodeGenPaths& TargetPaths);
+
+	/**  */
+	bool IsValid();
+
+public:
+	UPROPERTY()
+	FString GeneratedWrapperPath;
+
+	UPROPERTY()
+	TArray<FName> ReferencedBy;
+};
+
+/*******************************************************************************
+ * FBlueprintNativeCodeGenPaths
+ ******************************************************************************/
+
+struct FBlueprintNativeCodeGenPaths
+{
+public:
+	enum ESourceFileType
+	{
+		HFile, CppFile,
+	};
+
+	static FString GetDefaultManifestPath();
+
+private:
+	friend struct FBlueprintNativeCodeGenManifest;
+	FBlueprintNativeCodeGenPaths(const FString& PluginName, const FString& TargetDir, const FString& ManifestPath);
+
+public:
+	/**  */
+	FString ManifestFilePath() const;
+
+	/**  */
+	FString PluginRootDir() const;
+	FString PluginFilePath() const;
+	FString PluginSourceDir() const;
+
+	/**  */
+	FString RuntimeModuleDir() const;
+	FString RuntimeModuleName() const;
+	FString RuntimeBuildFile() const;
+	FString RuntimeSourceDir(ESourceFileType SourceType) const;
+	FString RuntimeModuleFile(ESourceFileType SourceType) const;
+	FString RuntimePCHFilename() const;
+
+private:
+	FString TargetDir;
+	FString PluginName;
+	FString ManifestPath;
+};
+
 /*******************************************************************************
  * FBlueprintNativeCodeGenManifest
  ******************************************************************************/
@@ -51,36 +122,19 @@ struct FBlueprintNativeCodeGenManifest
 {
 	GENERATED_USTRUCT_BODY()
 
-	/** @return The default filename for manifest files. */
-	static FString GetDefaultFilename();
+	typedef FName FAssetId;
+	typedef TMap<FAssetId, FConvertedAssetRecord> FConversionRecord;
+	typedef TMap<FAssetId, FUnconvertedDependencyRecord> FUnconvertedRecord;
 
 public: 
-	FBlueprintNativeCodeGenManifest(const FString TargetPath = TEXT(""));
+	FBlueprintNativeCodeGenManifest(); // default ctor for USTRUCT() system
 	FBlueprintNativeCodeGenManifest(const FNativeCodeGenCommandlineParams& CommandlineParams);
+	FBlueprintNativeCodeGenManifest(const FString& ManifestFilePath);
 
 	/**
-	 * A query that retrieves a list of file/directory paths. Conceptually, 
-	 * these files/directories are fully owned by the the conversion process.
-	 * 
-	 * @return An array of file/directory paths that will be generated the conversion process.
-	 */
-	TArray<FString> GetDestinationPaths() const;
-
-	/**
-	 * Constructs a .uplugin path, specifying the destination for the plugin's
-	 * descriptor file.
-	 * 
-	 * @return The destination filepath for the .uplugin file.
-	 */
-	FString GetPluginFilePath() const;
-
-	/**
-	 * Constructs a .Build.cs path, detailing the destination for the module's 
-	 * build file.
-	 * 
-	 * @return The destination filepath for the .Build.cs file.
+	 * @return A utility object that you can query for various file/directory paths used/targeted by the conversion process.
 	 */
-	FString GetModuleFilePath() const;
+	FBlueprintNativeCodeGenPaths GetTargetPaths() const;
 
 	/**
 	 * Logs an entry detailing the specified asset's conversion (the asset's 
@@ -90,7 +144,17 @@ public:
 	 * @param  AssetInfo    Defines the asset that you plan on converting (and want to make an entry for).
 	 * @return A record detailing the asset's expected conversion.
 	 */
-	FConvertedAssetRecord& CreateConversionRecord(const FAssetData& AssetInfo);
+	FConvertedAssetRecord& CreateConversionRecord(const FAssetId Key, const FAssetData& AssetInfo);
+
+	/**
+	 * 
+	 * 
+	 * @param  UnconvertedAssetKey    
+	 * @param  AssetObj    
+	 * @param  ReferencingAsset    
+	 * @return 
+	 */
+	FUnconvertedDependencyRecord& CreateUnconvertedDependencyRecord(const FAssetId UnconvertedAssetKey, const FAssetData& AssetInfo, const FAssetId ReferencingAsset);
 
 	/**
 	 * Searches the manifest's conversion records for an entry pertaining to the
@@ -99,29 +163,44 @@ public:
 	 * @param  AssetPath    Serves as a lookup key for the asset in question.
 	 * @return A pointer to the found conversion record, null if one doesn't exist.
 	 */
-	FConvertedAssetRecord* FindConversionRecord(const FString& AssetPath, bool bSlowSearch = false);
+	const FConvertedAssetRecord* FindConversionRecord(const FAssetId AssetId) const;
 
 	/**
 	 * Records module dependencies for the specified asset.
 	 * 
-	 * @param  AssetInfo    The asset you want to collect dependencies for.
-	 * @return True if we were able to successfully collect dependencies from the asset; false if the process failed.
+	 * @param  Package    The asset you want to collect dependencies for.
 	 */
-	bool LogDependencies(const FAssetData& AssetInfo);
+	void GatherModuleDependencies(UPackage* Package);
 
 	/** 
 	 * @return A list of all known modules that this plugin will depend on. 
 	 */
 	const TArray<UPackage*>& GetModuleDependencies() const { return ModuleDependencies; }
 
+	/** 
+	 * @return A list of all asset conversion that have been recorded. 
+	 */
+	const FConversionRecord& GetConversionRecord() const { return ConvertedAssets; }
+
 	/**
 	 * Saves this manifest as json, to its target destination (which it was 
 	 * setup with).
 	 * 
-	 * @param  bSort    Defines if you want the manifest's conversion log sorted before the files is saved (making it easier to read).
 	 * @return True if the save was a success, otherwise false.
 	 */
-	bool Save(bool bSort = true) const;
+	bool Save() const;
+
+	/**
+	* Saves this manifest as json, to the provided destination 
+	*
+	* @return True if the save was a success, otherwise false.
+	*/
+	bool SaveAs(const TCHAR* Filename) const;
+
+	/**
+	* Merges the manifest saved in Filename into the current manifest
+	*/
+	void Merge(const TCHAR* Filename);
 
 private:
 	/**
@@ -129,22 +208,11 @@ private:
 	 */
 	FString GetTargetDir() const;
 
-	/**
-	 * @return The target name for the plugin's primary module.
-	 */
-	const FString& GetTargetModuleName() const { return PluginName; }
-
 	/**
 	 * Empties the manifest, leaving only the destination directory and file
 	 * names intact.
 	 */
 	void Clear();
-
-	/**
-	 * Maps the asset log so that we can easily query for specific conversion
-	 * records.
-	 */
-	void ConstructRecordLookupTable();
 	
 private:
 	/** Defines the destination filepath for this manifest */
@@ -160,10 +228,10 @@ private:
 	UPROPERTY()
 	TArray<UPackage*> ModuleDependencies;
 
-	/** Mutable so Save() can sort the array when requested */
+	/**  */
 	UPROPERTY()
-	mutable TArray<FConvertedAssetRecord> ConvertedAssets;
+	TMap<FName, FConvertedAssetRecord> ConvertedAssets;
 
-	/** Serves as a quick way to look up records in the ConvertedAssets array */
-	mutable TMap<FString, int32> RecordLookupTable;
+	UPROPERTY()
+	TMap<FName, FUnconvertedDependencyRecord> UnconvertedDependencies;
 };
