Copying //UE4/Dev-Blueprints to Dev-Main (//UE4/Dev-Main) @ 2781164

diff --git a/Engine/Source/Developer/BlueprintNativeCodeGen/Private/BlueprintNativeCodeGenModule.cpp b/Engine/Source/Developer/BlueprintNativeCodeGen/Private/BlueprintNativeCodeGenModule.cpp
--- a/Engine/Source/Developer/BlueprintNativeCodeGen/Private/BlueprintNativeCodeGenModule.cpp
+++ b/Engine/Source/Developer/BlueprintNativeCodeGen/Private/BlueprintNativeCodeGenModule.cpp
@@ -1,21 +1,287 @@
 // Copyright 1998-2015 Epic Games, Inc. All Rights Reserved.
 
 #include "BlueprintNativeCodeGenPCH.h"
+
+#include "AssetRegistryModule.h"
+#include "BlueprintNativeCodeGenManifest.h"
 #include "BlueprintNativeCodeGenModule.h"
+#include "BlueprintNativeCodeGenUtils.h"
+#include "IBlueprintCompilerCppBackendModule.h"
+#include "NativeCodeGenCommandlineParams.h"
+
+/*******************************************************************************
+* NativizationCookControllerImpl
+******************************************************************************/
 
-/**
- *
- */
-class FBlueprintNativeCodeGenModule : public IBlueprintNativeCodeGenModuleInterface
+namespace NativizationCookControllerImpl
 {
-public:
-	virtual void StartupModule() override
+#if WITH_EDITORONLY_DATA
+	static const FString ConvertAssetsCommand = TEXT("NativizeAssets");
+	static const FString ConvertedManifestParam = TEXT("NativizedAssetManifest=");
+	static const FString GeneratedPluginParam = TEXT("NativizedAssetPlugin=");
+
+	//--------------------------------------------------------------------------
+	static bool IsRunningCookCommandlet()
 	{
+		// @TODO: figure out how to determine if this is the cook commandlet
+		return IsRunningCommandlet();
 	}
+#endif // WITH_EDITORONLY_DATA
+}
+
+/*******************************************************************************
+* UBlueprintNativeCodeGenConfig
+*******************************************************************************/
+
+//------------------------------------------------------------------------------
+UBlueprintNativeCodeGenConfig::UBlueprintNativeCodeGenConfig(FObjectInitializer const& ObjectInitializer)
+: Super(ObjectInitializer)
+{
+}
+
+/*******************************************************************************
+* FCookCommandParams
+******************************************************************************/
 
-	virtual void ShutdownModule() override
+//------------------------------------------------------------------------------
+FCookCommandParams::FCookCommandParams(const FString& Commandline)
+: bRunConversion(false)
+{
+	using namespace NativizationCookControllerImpl;
+
+#if WITH_EDITORONLY_DATA
+	if (IsRunningCookCommandlet())
 	{
+		IFileManager& FileManager = IFileManager::Get();
+		bRunConversion = FParse::Param(*Commandline, *ConvertAssetsCommand);
+
+		if (FParse::Value(*Commandline, *ConvertedManifestParam, ManifestFilePath))
+		{
+			bRunConversion |= !FileManager.FileExists(*ManifestFilePath);
+		}
+
+		FString PluginPath;
+		if (FParse::Value(*Commandline, *GeneratedPluginParam, PluginPath))
+		{
+			if (PluginPath.EndsWith(TEXT(".uplugin")))
+			{
+				PluginName = FPaths::GetBaseFilename(PluginPath);
+				PluginPath = FPaths::GetPath(PluginPath);
+			}
+			OutputPath = PluginPath;
+			bRunConversion = true;
+		}
+
+		if (bRunConversion && ManifestFilePath.IsEmpty())
+		{
+			ManifestFilePath = FPaths::Combine(*FPaths::Combine(*FPaths::GameIntermediateDir(), TEXT("Cook")), TEXT("BlueprintConversionManifest"));
+		}
 	}
+#endif // WITH_EDITORONLY_DATA
+}
+
+//------------------------------------------------------------------------------
+TArray<FString> FCookCommandParams::ToConversionParams() const
+{
+	TArray<FString> ConversionCmdParams;
+	if (!PluginName.IsEmpty())
+	{
+		ConversionCmdParams.Add(FString::Printf(TEXT("pluginName=%s"), *PluginName));
+	}
+	if (!OutputPath.IsEmpty())
+	{
+		ConversionCmdParams.Add(FString::Printf(TEXT("output=\"%s\""), *OutputPath));
+	}
+	if (!ManifestFilePath.IsEmpty())
+	{
+		ConversionCmdParams.Add(FString::Printf(TEXT("manifest=\"%s\""), *ManifestFilePath));
+	}
+	return ConversionCmdParams;
+}
+
+/*******************************************************************************
+ * FBlueprintNativeCodeGenModule
+ ******************************************************************************/
+ 
+class FBlueprintNativeCodeGenModule : public IBlueprintNativeCodeGenModule
+{
+public:
+	FBlueprintNativeCodeGenModule()
+		: Manifest(nullptr)
+	{
+		FCookCommandParams CookParams(FCommandLine::Get());
+		FNativeCodeGenCommandlineParams CommandLineParams(CookParams.ToConversionParams());
+		Manifest = new FBlueprintNativeCodeGenManifest(CommandLineParams);
+
+		IBlueprintCompilerCppBackendModule& BackEndModule = (IBlueprintCompilerCppBackendModule&)IBlueprintCompilerCppBackendModule::Get();
+		auto& BackendPCHQuery = BackEndModule.OnPCHFilenameQuery();
+
+		FBlueprintNativeCodeGenPaths TargetPaths = GetManifest().GetTargetPaths();
+		BackendPCHQuery.BindLambda([TargetPaths]()->FString
+		{
+			return TargetPaths.RuntimePCHFilename();
+		});
+
+		auto& ConversionQueryDelegate = BackEndModule.OnIsTargetedForConversionQuery();
+		auto ShouldConvert = [](const UObject* AssetObj)
+		{
+			if (FScriptCookReplacementCoordinator::Get())
+			{
+				EReplacementResult ReplacmentResult = FScriptCookReplacementCoordinator::Get()->IsTargetedForReplacement(AssetObj);
+				return ReplacmentResult == EReplacementResult::ReplaceCompletely;
+			}
+			return false;
+		};
+		ConversionQueryDelegate.BindStatic(ShouldConvert);
+	}
+
+	//~ Begin IBlueprintNativeCodeGenModule interface
+	virtual void Convert(UPackage* Package, EReplacementResult ReplacementType) override;
+	virtual void SaveManifest(const TCHAR* Filename) override;
+	virtual void MergeManifest(const TCHAR* Filename) override;
+	virtual void FinalizeManifest() override;
+	//~ End IBlueprintNativeCodeGenModule interface
+private:
+	FBlueprintNativeCodeGenManifest& GetManifest();
+
+	FBlueprintNativeCodeGenManifest* Manifest;
 };
 
-IMPLEMENT_MODULE(FBlueprintNativeCodeGenModule, BlueprintNativeCodeGen);
+FBlueprintNativeCodeGenManifest& FBlueprintNativeCodeGenModule::GetManifest()
+{
+	check(Manifest);
+	return *Manifest;
+}
+
+void FBlueprintNativeCodeGenModule::Convert(UPackage* Package, EReplacementResult ReplacementType)
+{
+	// Find the struct/enum to convert:
+	UStruct* Struct = nullptr;
+	UEnum* Enum = nullptr;
+	TArray<UObject*> Objects;
+	GetObjectsWithOuter(Package, Objects, false);
+	for (auto Entry : Objects)
+	{
+		if (Entry->HasAnyFlags(RF_Transient))
+		{
+			continue;
+		}
+
+		Struct = Cast<UStruct>(Entry);
+		if (Struct)
+		{
+			break;
+		}
+
+		Enum = Cast<UEnum>(Entry);
+		if (Enum)
+		{
+			break;
+		}
+	}
+
+	if (Struct == nullptr && Enum == nullptr)
+	{
+		check(false);
+		return;
+	}
+
+	const IAssetRegistry& Registry = FModuleManager::LoadModuleChecked<FAssetRegistryModule>("AssetRegistry").Get();
+	auto GenerateWrapperSrcFile = [](const FUnconvertedDependencyRecord& UnconvertedDependency, UClass* Class)
+	{
+		FString FileContents;
+		TUniquePtr<IBlueprintCompilerCppBackend> Backend_CPP(IBlueprintCompilerCppBackendModuleInterface::Get().Create());
+		FileContents = Backend_CPP->GenerateWrapperForClass(Class);
+
+		if (!FileContents.IsEmpty())
+		{
+			FFileHelper::SaveStringToFile(FileContents, *UnconvertedDependency.GeneratedWrapperPath);
+		}
+	};
+
+	if (ReplacementType == EReplacementResult::GenerateStub)
+	{
+		check(ReplacementType == EReplacementResult::GenerateStub);
+		FAssetData AssetInfo = Registry.GetAssetByObjectPath(*Struct->GetPathName());
+		GenerateWrapperSrcFile(GetManifest().CreateUnconvertedDependencyRecord(AssetInfo.PackageName, AssetInfo, FName()), CastChecked<UClass>(Struct));
+		// The stub we generate still may have dependencies on other modules, so make sure the module dependencies are 
+		// still recorded so that the .build.cs is generated correctly. Without this you'll get include related errors 
+		// (or possibly linker errors) in stub headers:
+		GetManifest().GatherModuleDependencies(Package);
+	}
+	else
+	{
+		check(ReplacementType == EReplacementResult::ReplaceCompletely);
+		// convert:
+		UField* ForConversion = Enum;
+		if (ForConversion == nullptr)
+		{
+			ForConversion = Struct;
+		}
+
+		FAssetData AssetInfo = Registry.GetAssetByObjectPath(*ForConversion->GetPathName());
+		FConvertedAssetRecord& ConversionRecord = GetManifest().CreateConversionRecord(ForConversion->GetFName(), AssetInfo);
+		TSharedPtr<FString> HeaderSource(new FString());
+		TSharedPtr<FString> CppSource(new FString());
+
+		FBlueprintNativeCodeGenUtils::GenerateCppCode(ForConversion, HeaderSource, CppSource);
+		bool bSuccess = !HeaderSource->IsEmpty() || !CppSource->IsEmpty();
+		// Run the cpp first, because we cue off of the presence of a header for a valid conversion record (see
+		// FConvertedAssetRecord::IsValid)
+		if (!CppSource->IsEmpty())
+		{
+			if (!FFileHelper::SaveStringToFile(*CppSource, *ConversionRecord.GeneratedCppPath))
+			{
+				bSuccess &= false;
+				ConversionRecord.GeneratedCppPath.Empty();
+			}
+			CppSource->Empty(CppSource->Len());
+		}
+		else
+		{
+			ConversionRecord.GeneratedCppPath.Empty();
+		}
+
+		if (bSuccess && !HeaderSource->IsEmpty())
+		{
+			if (!FFileHelper::SaveStringToFile(*HeaderSource, *ConversionRecord.GeneratedHeaderPath))
+			{
+				bSuccess &= false;
+				ConversionRecord.GeneratedHeaderPath.Empty();
+			}
+			HeaderSource->Empty(HeaderSource->Len());
+		}
+		else
+		{
+			ConversionRecord.GeneratedHeaderPath.Empty();
+		}
+
+		check(bSuccess);
+		if (bSuccess)
+		{
+			GetManifest().GatherModuleDependencies(Package);
+		}
+	}
+}
+
+void FBlueprintNativeCodeGenModule::SaveManifest(const TCHAR* Filename)
+{
+	GetManifest().SaveAs(Filename);
+}
+
+void FBlueprintNativeCodeGenModule::MergeManifest(const TCHAR* Filename)
+{
+	GetManifest().Merge(Filename);
+} 
+
+void FBlueprintNativeCodeGenModule::FinalizeManifest()
+{
+	GetManifest().Save();
+
+	FCookCommandParams CookParams(FCommandLine::Get());
+	FNativeCodeGenCommandlineParams CommandLineParams(CookParams.ToConversionParams());
+	check(FPaths::IsSamePath(CommandLineParams.ManifestFilePath, CommandLineParams.ManifestFilePath));
+	check(FBlueprintNativeCodeGenUtils::FinalizePlugin(GetManifest(), CommandLineParams));
+}
+
+IMPLEMENT_MODULE(FBlueprintNativeCodeGenModule, BlueprintNativeCodeGen);
\ No newline at end of file
diff --git a/Engine/Source/Developer/BlueprintNativeCodeGen/Private/BlueprintNativeCodeGenUtils.cpp b/Engine/Source/Developer/BlueprintNativeCodeGen/Private/BlueprintNativeCodeGenUtils.cpp
--- a/Engine/Source/Developer/BlueprintNativeCodeGen/Private/BlueprintNativeCodeGenUtils.cpp
+++ b/Engine/Source/Developer/BlueprintNativeCodeGen/Private/BlueprintNativeCodeGenUtils.cpp
@@ -1,20 +1,26 @@
 // Copyright 1998-2015 Epic Games, Inc. All Rights Reserved.
 
 #include "BlueprintNativeCodeGenPCH.h"
+#include "BlueprintNativeCodeGenManifest.h"
 #include "BlueprintNativeCodeGenUtils.h"
 #include "NativeCodeGenCommandlineParams.h"
-#include "BlueprintNativeCodeGenCoordinator.h"
-#include "Kismet2/KismetReinstanceUtilities.h"	// for FBlueprintCompileReinstancer
+#include "Kismet2/KismetReinstanceUtilities.h"	 // for FBlueprintCompileReinstancer
 #include "Kismet2/CompilerResultsLog.h" 
 #include "KismetCompilerModule.h"
 #include "Engine/Blueprint.h"
 #include "Engine/UserDefinedStruct.h"
 #include "Engine/UserDefinedEnum.h"
-#include "Kismet2/KismetEditorUtilities.h"		// for CompileBlueprint()
-#include "OutputDevice.h"						// for GWarn
-#include "GameProjectUtils.h"					// for GenerateGameModuleBuildFile
+#include "Kismet2/KismetEditorUtilities.h"		 // for CompileBlueprint()
+#include "OutputDevice.h"						 // for GWarn
+#include "GameProjectUtils.h"					 // for GenerateGameModuleBuildFile
 #include "Editor/GameProjectGeneration/Public/GameProjectUtils.h" // for GenerateGameModuleBuildFile()
-#include "App.h"								// for GetGameName()
+#include "App.h"								 // for GetGameName()
+#include "BlueprintSupport.h"					 // for FReplaceCookedBPGC
+#include "IBlueprintCompilerCppBackendModule.h"	 // for OnPCHFilenameQuery()
+#include "BlueprintNativeCodeGenModule.h"
+#include "ScopeExit.h"
+#include "Editor/Kismet/Public/FindInBlueprintManager.h" // for FDisableGatheringDataOnScope
+#include "Editor/UnrealEd/Public/Kismet2/BlueprintEditorUtils.h" // for FBlueprintEditorUtils::FForceFastBlueprintDuplicationScope
 
 DEFINE_LOG_CATEGORY(LogBlueprintCodeGen)
 
@@ -24,55 +30,70 @@ DEFINE_LOG_CATEGORY(LogBlueprintCodeGen)
 
 namespace BlueprintNativeCodeGenUtilsImpl
 {
+	static FString CoreModuleName   = TEXT("Core");
+	static FString EngineModuleName = TEXT("Engine");
+	static FString EngineHeaderFile = TEXT("Engine.h");
+
 	/**
 	 * Deletes the files/directories in the supplied array.
 	 * 
 	 * @param  TargetPaths    The set of directory and file paths that you want deleted.
 	 * @return True if all the files/directories were successfully deleted, other wise false.
 	 */
-	static bool WipeTargetPaths(const TArray<FString>& TargetPaths);
+	static bool WipeTargetPaths(const FBlueprintNativeCodeGenPaths& TargetPaths);
 	
 	/**
 	 * Creates and fills out a new .uplugin file for the converted assets.
 	 * 
 	 * @param  PluginName	The name of the plugin you're generating.
-	 * @param  Manifest		Defines where the plugin file should be saved.
+	 * @param  TargetPaths	Defines the file path/name for the plugin file.
 	 * @return True if the file was successfully saved, otherwise false.
 	 */
-	static bool GeneratePluginDescFile(const FString& PluginName, const FBlueprintNativeCodeGenManifest& Manifest);
+	static bool GeneratePluginDescFile(const FString& PluginName, const FBlueprintNativeCodeGenPaths& TargetPaths);
+
+	/**
+	 * Creates a module implementation and header file for the converted assets'
+	 * module (provides a IMPLEMENT_MODULE() declaration, which is required for 
+	 * the module to function).
+	 * 
+	 * @param  TargetPaths    Defines the file path/name for the target files.
+	 * @return True if the files were successfully generated, otherwise false.
+	 */
+	static bool GenerateModuleSourceFiles(const FBlueprintNativeCodeGenPaths& TargetPaths);
 
 	/**
-	 * Creates and fills out a new .Build.cs file for the target module.
+	 * Creates and fills out a new .Build.cs file for the plugin's runtime module.
 	 * 
 	 * @param  Manifest    Defines where the module file should be saved, what it should be named, etc..
 	 * @return True if the file was successfully saved, otherwise false.
 	 */
 	static bool GenerateModuleBuildFile(const FBlueprintNativeCodeGenManifest& Manifest);
+
+	/**
+	 * Determines what the expected native class will be for an asset that was  
+	 * or will be converted.
+	 * 
+	 * @param  ConversionRecord    Identifies the asset's original type (which is used to infer the replacement type from).
+	 * @return Either a class, enum, or struct class (depending on the asset's type).
+	 */
+	static UClass* ResolveReplacementType(const FConvertedAssetRecord& ConversionRecord);
 }
 
 //------------------------------------------------------------------------------
-static bool BlueprintNativeCodeGenUtilsImpl::WipeTargetPaths(const TArray<FString>& TargetPaths)
+static bool BlueprintNativeCodeGenUtilsImpl::WipeTargetPaths(const FBlueprintNativeCodeGenPaths& TargetPaths)
 {
 	IFileManager& FileManager = IFileManager::Get();
 
 	bool bSuccess = true;
-	for (const FString& Path : TargetPaths)
-	{
-		if (FileManager.FileExists(*Path))
-		{
-			bSuccess &= FileManager.Delete(*Path);
-		}
-		else if (FileManager.DirectoryExists(*Path))
-		{
-			bSuccess &= FileManager.DeleteDirectory(*Path, /*RequireExists =*/false, /*Tree =*/true);
-		}
-	}
+	bSuccess &= FileManager.Delete(*TargetPaths.PluginFilePath());
+	bSuccess &= FileManager.DeleteDirectory(*TargetPaths.PluginSourceDir(), /*RequireExists =*/false, /*Tree =*/true);
+	bSuccess &= FileManager.Delete(*TargetPaths.ManifestFilePath());
 
 	return bSuccess;
 }
 
 //------------------------------------------------------------------------------
-static bool BlueprintNativeCodeGenUtilsImpl::GeneratePluginDescFile(const FString& PluginName, const FBlueprintNativeCodeGenManifest& Manifest)
+static bool BlueprintNativeCodeGenUtilsImpl::GeneratePluginDescFile(const FString& PluginName, const FBlueprintNativeCodeGenPaths& TargetPaths)
 {
 	FPluginDescriptor PluginDesc;
 	PluginDesc.FriendlyName = PluginName;
@@ -86,16 +107,16 @@ static bool BlueprintNativeCodeGenUtilsImpl::GeneratePluginDescFile(const FStrin
 	PluginDesc.bCanContainContent = false;
 	PluginDesc.bIsBetaVersion     = true; // @TODO: change once we're confident in the feature
 
-	FModuleDescriptor ModuleDesc;
-	ModuleDesc.Name = *PluginName;
-	ModuleDesc.Type = EHostType::Runtime;
+	FModuleDescriptor RuntimeModuleDesc;
+	RuntimeModuleDesc.Name = *TargetPaths.RuntimeModuleName();
+	RuntimeModuleDesc.Type = EHostType::Runtime;
 	// load at startup (during engine init), after game modules have been loaded 
-	ModuleDesc.LoadingPhase = ELoadingPhase::Default; 
+	RuntimeModuleDesc.LoadingPhase = ELoadingPhase::Default;
 
-	PluginDesc.Modules.Add(ModuleDesc);
+	PluginDesc.Modules.Add(RuntimeModuleDesc);
 
 	FText ErrorMessage;
-	bool bSuccess = PluginDesc.Save(Manifest.GetPluginFilePath(), ErrorMessage);
+	bool bSuccess = PluginDesc.Save(TargetPaths.PluginFilePath(), ErrorMessage);
 
 	if (!bSuccess)
 	{
@@ -104,12 +125,41 @@ static bool BlueprintNativeCodeGenUtilsImpl::GeneratePluginDescFile(const FStrin
 	return bSuccess;
 }
 
+//------------------------------------------------------------------------------
+static bool BlueprintNativeCodeGenUtilsImpl::GenerateModuleSourceFiles(const FBlueprintNativeCodeGenPaths& TargetPaths)
+{
+	FText FailureReason;
+
+	TArray<FString> PchIncludes;
+	PchIncludes.Add(EngineHeaderFile);
+
+	bool bSuccess = GameProjectUtils::GeneratePluginModuleHeaderFile(TargetPaths.RuntimeModuleFile(FBlueprintNativeCodeGenPaths::HFile), PchIncludes, FailureReason);
+
+	if (bSuccess)
+	{
+		const FString NoStartupCode = TEXT("");
+		bSuccess &= GameProjectUtils::GeneratePluginModuleCPPFile(TargetPaths.RuntimeModuleFile(FBlueprintNativeCodeGenPaths::CppFile),
+			TargetPaths.RuntimeModuleName(), NoStartupCode, FailureReason);
+	}
+
+	if (!bSuccess)
+	{
+		UE_LOG(LogBlueprintCodeGen, Error, TEXT("Failed to generate module source files: %s"), *FailureReason.ToString());
+	}
+	return bSuccess;
+}
+
 //------------------------------------------------------------------------------
 static bool BlueprintNativeCodeGenUtilsImpl::GenerateModuleBuildFile(const FBlueprintNativeCodeGenManifest& Manifest)
 {
 	FModuleManager& ModuleManager = FModuleManager::Get();
 	
 	TArray<FString> PublicDependencies;
+	// for IModuleInterface
+	PublicDependencies.Add(CoreModuleName);
+	// for Engine.h
+	PublicDependencies.Add(EngineModuleName);
+
 	if (GameProjectUtils::ProjectHasCodeFiles()) 
 	{
 		const FString GameModuleName = FApp::GetGameName();
@@ -137,15 +187,11 @@ static bool BlueprintNativeCodeGenUtilsImpl::GenerateModuleBuildFile(const FBlue
 		}
 	}
 
-	const FString BuildFilePath = Manifest.GetModuleFilePath();
-
-	FString ModuleName = FPaths::GetCleanFilename(BuildFilePath);
-	int32 ExtIndex = INDEX_NONE;
-	ModuleName.FindChar('.', ExtIndex);
-	ModuleName = ModuleName.Left(ExtIndex);
+	FBlueprintNativeCodeGenPaths TargetPaths = Manifest.GetTargetPaths();
 
 	FText ErrorMessage;
-	bool bSuccess = GameProjectUtils::GenerateGameModuleBuildFile(BuildFilePath, ModuleName, PublicDependencies, PrivateDependencies, ErrorMessage);
+	bool bSuccess = GameProjectUtils::GenerateGameModuleBuildFile(TargetPaths.RuntimeBuildFile(), TargetPaths.RuntimeModuleName(),
+		PublicDependencies, PrivateDependencies, ErrorMessage);
 
 	if (!bSuccess)
 	{
@@ -154,98 +200,43 @@ static bool BlueprintNativeCodeGenUtilsImpl::GenerateModuleBuildFile(const FBlue
 	return bSuccess;
 }
 
-/*******************************************************************************
- * FBlueprintNativeCodeGenUtils
- ******************************************************************************/
-
 //------------------------------------------------------------------------------
-bool FBlueprintNativeCodeGenUtils::GeneratePlugin(const FNativeCodeGenCommandlineParams& CommandParams)
+static UClass* BlueprintNativeCodeGenUtilsImpl::ResolveReplacementType(const FConvertedAssetRecord& ConversionRecord)
 {
-	FScopedFeedbackContext ScopedErrorTracker;
-	FBlueprintNativeCodeGenCoordinator Coordinator(CommandParams);
+	const UClass* AssetType = ConversionRecord.AssetType;
+	check(AssetType != nullptr);
 
-	if (ScopedErrorTracker.HasErrors())
+	if (AssetType->IsChildOf<UUserDefinedEnum>())
 	{
-		// creating the coordinator/manifest produced an error... do not carry on!
-		return false;
+		return UEnum::StaticClass();
 	}
-	const FBlueprintNativeCodeGenManifest& Manifest = Coordinator.GetManifest();
-
-	if (CommandParams.bWipeRequested && !CommandParams.bPreviewRequested)
+	else if (AssetType->IsChildOf<UUserDefinedStruct>())
 	{
-		TArray<FString> TargetPaths = Manifest.GetDestinationPaths();
-		if (!BlueprintNativeCodeGenUtilsImpl::WipeTargetPaths(TargetPaths))
-		{
-			UE_LOG(LogBlueprintCodeGen, Warning, TEXT("Failed to wipe target files/directories."));
-		}
+		return UScriptStruct::StaticClass();
 	}
-
-	TSharedPtr<FString> HeaderSource(new FString());
-	TSharedPtr<FString> CppSource(new FString());
-
-	auto ConvertSingleAsset = [&Coordinator, &HeaderSource, &CppSource](FConvertedAssetRecord& ConversionRecord)->bool
+	else if (AssetType->IsChildOf<UBlueprint>())
 	{
-		FScopedFeedbackContext NestedErrorTracker;
-
-		UObject* AssetObj = ConversionRecord.AssetPtr.Get();
-		if (AssetObj != nullptr)
-		{
-			FBlueprintNativeCodeGenUtils::GenerateCppCode(AssetObj, HeaderSource, CppSource);
-		}
-
-		bool bSuccess = !HeaderSource->IsEmpty() || !CppSource->IsEmpty();
-		// run the cpp first, because we cue off of the presence of a header for 
-		// a valid conversion record (see FConvertedAssetRecord::IsValid)
-		if (!CppSource->IsEmpty())
-		{
-			if (!FFileHelper::SaveStringToFile(*CppSource, *ConversionRecord.GeneratedCppPath))
-			{
-				bSuccess &= false;
-				ConversionRecord.GeneratedCppPath.Empty();
-			}
-			CppSource->Empty(CppSource->Len());
-		}
-		else
-		{
-			ConversionRecord.GeneratedCppPath.Empty();
-		}
-
-		if (bSuccess && !HeaderSource->IsEmpty())
-		{
-			if (!FFileHelper::SaveStringToFile(*HeaderSource, *ConversionRecord.GeneratedHeaderPath))
-			{
-				bSuccess &= false;
-				ConversionRecord.GeneratedHeaderPath.Empty();
-			}
-			HeaderSource->Empty(HeaderSource->Len());
-		}
-		else
-		{
-			ConversionRecord.GeneratedHeaderPath.Empty();
-		}
-		return bSuccess && !NestedErrorTracker.HasErrors();
-	};
-
-	FBlueprintNativeCodeGenCoordinator::FConversionDelegate ConversionDelegate = FBlueprintNativeCodeGenCoordinator::FConversionDelegate::CreateLambda(ConvertSingleAsset);
-	if (CommandParams.bPreviewRequested)
-	{
-		ConversionDelegate.BindLambda([](FConvertedAssetRecord& ConversionRecord){ return true; });
+		return UDynamicClass::StaticClass();
 	}
-
-	bool bSuccess = Coordinator.ProcessConversionQueue(ConversionDelegate);
-
-	if (bSuccess && !CommandParams.bPreviewRequested)
+	else
 	{
-		bSuccess &= BlueprintNativeCodeGenUtilsImpl::GenerateModuleBuildFile(Manifest);
-		if (bSuccess)
-		{
-			bSuccess &= BlueprintNativeCodeGenUtilsImpl::GeneratePluginDescFile(CommandParams.PluginName, Manifest);
-		}
+		UE_LOG(LogBlueprintCodeGen, Error, TEXT("Unsupported asset type (%s); cannot determine replacement type."), *AssetType->GetName());
 	}
+	return nullptr;
+}
 
-	// save the manifest regardless of success (so we can get an idea of where failures broke down)
-	Manifest.Save();
-	return bSuccess & !ScopedErrorTracker.HasErrors();
+/*******************************************************************************
+ * FBlueprintNativeCodeGenUtils
+ ******************************************************************************/
+
+//------------------------------------------------------------------------------
+bool FBlueprintNativeCodeGenUtils::FinalizePlugin(const FBlueprintNativeCodeGenManifest& Manifest, const FNativeCodeGenCommandlineParams& CommandParams)
+{
+	bool bSuccess = true;
+	bSuccess = bSuccess && BlueprintNativeCodeGenUtilsImpl::GenerateModuleBuildFile(Manifest);
+	bSuccess = bSuccess && BlueprintNativeCodeGenUtilsImpl::GenerateModuleSourceFiles(Manifest.GetTargetPaths());
+	bSuccess = bSuccess && BlueprintNativeCodeGenUtilsImpl::GeneratePluginDescFile(CommandParams.PluginName, Manifest.GetTargetPaths());
+	return bSuccess;
 }
 
 //------------------------------------------------------------------------------
@@ -261,32 +252,67 @@ void FBlueprintNativeCodeGenUtils::GenerateCppCode(UObject* Obj, TSharedPtr<FStr
 
 	if (InBlueprintObj)
 	{
+		if (EBlueprintStatus::BS_Error == InBlueprintObj->Status)
+		{
+			UE_LOG(LogBlueprintCodeGen, Error, TEXT("Cannot convert \"%s\". It has errors."), *InBlueprintObj->GetPathName());
+			return;
+		}
+
 		check(InBlueprintObj->GetOutermost() != GetTransientPackage());
-		checkf(InBlueprintObj->GeneratedClass, TEXT("Invalid generated class for %s"), *InBlueprintObj->GetName());
+		if (!ensureMsgf(InBlueprintObj->GeneratedClass, TEXT("Invalid generated class for %s"), *InBlueprintObj->GetName()))
+		{
+			return;
+		}
 		check(OutHeaderSource.IsValid());
 		check(OutCppSource.IsValid());
 
-		auto BlueprintObj = InBlueprintObj;
+		FDisableGatheringDataOnScope DisableFib;
+
+		const FString TempPackageName = FString::Printf(TEXT("/Temp/__TEMP_BP__/%s"), *InBlueprintObj->GetName());
+		UPackage* TempPackage = CreatePackage(nullptr, *TempPackageName);
+		check(TempPackage);
+		ON_SCOPE_EXIT
 		{
-			TSharedPtr<FBlueprintCompileReinstancer> Reinstancer = FBlueprintCompileReinstancer::Create(BlueprintObj->GeneratedClass);
+			TempPackage->RemoveFromRoot();
+			TempPackage->MarkPendingKill();
+		};
 
+		UBlueprint* DuplicateBP = nullptr;
+		{
+			FBlueprintDuplicationScopeFlags BPDuplicationFlags(
+				FBlueprintDuplicationScopeFlags::NoExtraCompilation | FBlueprintDuplicationScopeFlags::TheSameTimelineGuid);
+			DuplicateBP = DuplicateObject<UBlueprint>(InBlueprintObj, TempPackage, *InBlueprintObj->GetName());
+		}
+		ensure((nullptr != DuplicateBP->GeneratedClass) && (InBlueprintObj->GeneratedClass != DuplicateBP->GeneratedClass));
+		ON_SCOPE_EXIT
+		{
+			DuplicateBP->RemoveFromRoot();
+			DuplicateBP->MarkPendingKill();
+		};
+
+		IBlueprintCompilerCppBackendModule& CodeGenBackend = (IBlueprintCompilerCppBackendModule&)IBlueprintCompilerCppBackendModule::Get();
+		CodeGenBackend.GetOriginalClassMap().Add(*DuplicateBP->GeneratedClass, *InBlueprintObj->GeneratedClass);
+
+		{
+			TSharedPtr<FBlueprintCompileReinstancer> Reinstancer = FBlueprintCompileReinstancer::Create(DuplicateBP->GeneratedClass);
 			IKismetCompilerInterface& Compiler = FModuleManager::LoadModuleChecked<IKismetCompilerInterface>(KISMET_COMPILER_MODULENAME);
-			
 			TGuardValue<bool> GuardTemplateNameFlag(GCompilingBlueprint, true);
-			FCompilerResultsLog Results;			
+			FCompilerResultsLog Results;
 
 			FKismetCompilerOptions CompileOptions;
 			CompileOptions.CompileType = EKismetCompileType::Cpp;
 			CompileOptions.OutCppSourceCode = OutCppSource;
 			CompileOptions.OutHeaderSourceCode = OutHeaderSource;
-			Compiler.CompileBlueprint(BlueprintObj, CompileOptions, Results);
 
-			if (EBlueprintType::BPTYPE_Interface == BlueprintObj->BlueprintType && OutCppSource.IsValid())
-			{
-				OutCppSource->Empty(); // ugly temp hack
-			}
+			Compiler.CompileBlueprint(DuplicateBP, CompileOptions, Results);
+
+			Compiler.RemoveBlueprintGeneratedClasses(DuplicateBP);
+		}
+
+		if (EBlueprintType::BPTYPE_Interface == DuplicateBP->BlueprintType && OutCppSource.IsValid())
+		{
+			OutCppSource->Empty(); // ugly temp hack
 		}
-		FKismetEditorUtilities::CompileBlueprint(BlueprintObj);
 	}
 	else if ((UDEnum || UDStruct) && OutHeaderSource.IsValid())
 	{
diff --git a/Engine/Source/Developer/BlueprintNativeCodeGen/Private/BlueprintNativeCodeGenUtils.h b/Engine/Source/Developer/BlueprintNativeCodeGen/Private/BlueprintNativeCodeGenUtils.h
--- a/Engine/Source/Developer/BlueprintNativeCodeGen/Private/BlueprintNativeCodeGenUtils.h
+++ b/Engine/Source/Developer/BlueprintNativeCodeGen/Private/BlueprintNativeCodeGenUtils.h
@@ -7,6 +7,7 @@
 // Forward declares
 struct FNativeCodeGenCommandlineParams;
 class  SBuildProgressWidget;
+struct FBlueprintNativeCodeGenManifest;
 
 DECLARE_LOG_CATEGORY_EXTERN(LogBlueprintCodeGen, Log, All);
 
@@ -15,26 +16,26 @@ struct FBlueprintNativeCodeGenUtils
 {
 public:
 	/**
-	 * 
-	 * 
-	 * @param  CommandParams    
-	 * @return 
+	 * Generated module build file, module source files, and plugin description file for the provided manifest.
+	 *
+	 * @param	Manifest		The manifest containing lists of converted files, etc
+	 * @param	CommandParams	Command parameters used to generate native code, includes plugin name
 	 */
-	static bool GeneratePlugin(const FNativeCodeGenCommandlineParams& CommandParams);
+	static bool FinalizePlugin(const FBlueprintNativeCodeGenManifest& Manifest, const FNativeCodeGenCommandlineParams& CommandParams);
 
 	/**
 	 * Recompiles the bytecode of a blueprint only. Should only be run for
 	 * recompiling dependencies during compile on load.
 	 * 
-	 * @param  Obj    
-	 * @param  OutHeaderSource    
-	 * @param  OutCppSource    
+	 * @param  Obj				The asset object that you want to generate source code for (a Blueprint, or UserDefinedEnum/Struct)
+	 * @param  OutHeaderSource	Output destination for the header source.
+	 * @param  OutCppSource		Output destination for the cpp source.
 	 */
 	static void GenerateCppCode(UObject* Obj, TSharedPtr<FString> OutHeaderSource, TSharedPtr<FString> OutCppSource);
 
 public: 
 	/** 
-	 * 
+	 * A utility for catching errors/warnings that were logged in nested/scoped calls.
 	 */
 	class FScopedFeedbackContext : public FFeedbackContext
 	{
@@ -45,14 +46,14 @@ public:
 		/**  */
 		bool HasErrors();
 
-		// FOutputDevice interface
+		//~ Begin FOutputDevice interface
 		virtual void Serialize(const TCHAR* V, ELogVerbosity::Type Verbosity, const class FName& Category) override;
 		virtual void Flush() override;
 		virtual void TearDown() override { OldContext->TearDown();  }
 		virtual bool CanBeUsedOnAnyThread() const override { return false; }
-		// End FOutputDevice interface
+		//~ End FOutputDevice interface
 
-		// FFeedbackContext interface
+		//~ Begin FFeedbackContext interface
 		virtual bool YesNof(const FText& Question) override { return OldContext->YesNof(Question); }
 		virtual bool ReceivedUserCancel() override { return OldContext->ReceivedUserCancel(); }
 		virtual FContextSupplier* GetContext() const override { return OldContext->GetContext(); }
@@ -64,7 +65,7 @@ public:
 		//virtual void StartSlowTask(const FText& Task, bool bShowCancelButton = false) override { OldContext->StartSlowTask(Task, bShowCancelButton); }
 		//virtual void FinalizeSlowTask() override { OldContext->FinalizeSlowTask(); }
 		//virtual void ProgressReported(const float TotalProgressInterp, FText DisplayMessage) override { OldContext->ProgressReported(TotalProgressInterp, DisplayMessage); }
-		// End FFeedbackContext interface
+		//~ End FFeedbackContext interface
 
 	private:
 		FFeedbackContext* OldContext;
diff --git a/Engine/Source/Developer/BlueprintNativeCodeGen/Private/NativeCodeGenCommandlineParams.cpp b/Engine/Source/Developer/BlueprintNativeCodeGen/Private/NativeCodeGenCommandlineParams.cpp
--- a/Engine/Source/Developer/BlueprintNativeCodeGen/Private/NativeCodeGenCommandlineParams.cpp
+++ b/Engine/Source/Developer/BlueprintNativeCodeGen/Private/NativeCodeGenCommandlineParams.cpp
@@ -23,33 +23,42 @@ const FString FNativeCodeGenCommandlineParams::HelpMessage = TEXT("\n\
 :: Parameters                                                                  \n\
 ::                                                                             \n\
 \n\
-    -whitelist=<Assets> Identifies assets that you wish to convert. <Assets>   \n\
+   -whitelist=<Assets>  Identifies assets that you wish to convert. <Assets>   \n\
                         can be a comma seperated list, specifying multiple     \n\
                         packages (either directories, or explicit assets) in   \n\
                         the form: /Game/MyContentDir,/Engine/EngineAssetName.  \n\
 \n\
-    -blacklist=<Assets> Explicitly specifies assets that you don't want        \n\
+   -blacklist=<Assets>  Explicitly specifies assets that you don't want        \n\
                         converted (listed in the same manner as the -whitelist \n\
                         param). This takes priority over the whitelist (even if\n\
                         it results in uncompilable code).                      \n\
 \n\
-    -output=<TargetDir> Specifies the path where you want converted assets     \n\
+   -output=<TargetDir>  Specifies the path where you want converted assets     \n\
                         saved to. If left unset, this will default to the      \n\
                         project's intermeadiate folder (within a sub-folder of \n\
                         its own). If specified as a relative path, it will be  \n\
                         interpreted as relative to the project directory.      \n\
 \n\
-    -pluginName=<Name>  Defines the generated plugin's name. If the output     \n\
+   -pluginName=<Name>   Defines the generated plugin's name. If the output     \n\
                         directory doesn't match this in name, then a sub-folder\n\
                         will be created there to house the plugin.             \n\
 \n\
-    -noWipe             By default, we'll clear out the target directory from  \n\
+   -noWipe              By default, we'll clear out the target directory from  \n\
                         previous conversions. However, with this, if the       \n\
 						target directory contains source from a previous run,  \n\
-                        we'll build atop the existing module and appends to its\n\
+                        we'll build atop the existing module and append to its \n\
 						existing manifest.                                     \n\
 \n\
-    -manifest=<Path>    Specifies where to save the resultant manifest file. If\n\
+   -dependencies=<Rule> By default, we convert all required Blueprint          \n\
+                        dependencies (that are not blacklisted), even if they  \n\
+                        weren't found in the whitelist. You can also specify:  \n\
+                        'None' or 'All' (default is 'RequiredOnly'), which will\n\
+                        only convert whitelisted assets, or every dependency   \n\
+                        respectively. This param has no use if when a whitelist\n\
+                        is not provided (either through the commandline, or an \n\
+                        ini config file).                                      \n\
+\n\
+   -manifest=<Path>     Specifies where to save the resultant manifest file. If\n\
                         <Path> is not an existing directory or file, then it is\n\
                         assumed that this is a file path. If a manifest already\n\
                         exists here, then it'll be used to define the target   \n\
@@ -57,12 +66,12 @@ const FString FNativeCodeGenCommandlineParams::HelpMessage = TEXT("\n\
                         manifest file). If left unset then the manifest will   \n\
                         default to the plugin's directory.                     \n\
 \n\
-    -preview            When specified, the process will NOT wipe any existing \n\
+   -preview             When specified, the process will NOT wipe any existing \n\
                         files, and the target assets will NOT be compiled and  \n\
 						converted. Instead, just a manifest file will be       \n\
 						generated, detailing the command's predicted output.   \n\
 \n\
-    -help, -h, -?       Display this message and then exit.                    \n\
+   -help, -h, -?        Display this message and then exit.                    \n\
 \n");
 
 /*******************************************************************************
@@ -101,6 +110,7 @@ FNativeCodeGenCommandlineParams::FNativeCodeGenCommandlineParams(const TArray<FS
 	: bHelpRequested(false)
 	, bWipeRequested(true)
 	, bPreviewRequested(false)
+	, DependencyConversionRule(EDependencyConversionRule::RequiredOnly)
 {
 	using namespace NativeCodeGenCommandlineParamsImpl;
 
@@ -152,6 +162,21 @@ FNativeCodeGenCommandlineParams::FNativeCodeGenCommandlineParams(const TArray<FS
 		{
 			bWipeRequested = false;
 		}
+		else if (!Switch.Compare(TEXT("dependencies"), ESearchCase::IgnoreCase))
+		{
+			if (UEnum* const DependencyRuleEnum = FindObject<UEnum>(ANY_PACKAGE, TEXT("EDependencyConversionRule")))
+			{
+				FString EnumValueName = UEnum::GenerateFullEnumName(DependencyRuleEnum, *Value);
+				if (DependencyRuleEnum->IsValidEnumName(*EnumValueName))
+				{
+					DependencyConversionRule = (EDependencyConversionRule)DependencyRuleEnum->GetValueByName(*EnumValueName);
+				}
+				else
+				{
+					UE_LOG(LogNativeCodeGenCommandline, Warning, TEXT("Invalid dependency rule: '%s' (defaulting to 'RequiredOnly')."), *Value);
+				}
+			}
+		}
 		else if (!Switch.Compare(TEXT("preview"), ESearchCase::IgnoreCase))
 		{
 			bPreviewRequested = true;
@@ -166,7 +191,7 @@ FNativeCodeGenCommandlineParams::FNativeCodeGenCommandlineParams(const TArray<FS
 			}
 			else if (FileManager.DirectoryExists(*Value))
 			{
-				ManifestFilePath = FPaths::Combine(*Value, *FBlueprintNativeCodeGenManifest::GetDefaultFilename());
+				ManifestFilePath = FPaths::Combine(*Value, *FBlueprintNativeCodeGenPaths::GetDefaultManifestPath());
 			}
 			else
 			{
@@ -184,7 +209,7 @@ FNativeCodeGenCommandlineParams::FNativeCodeGenCommandlineParams(const TArray<FS
 						Value = FPaths::ConvertRelativePathToFull(Value);
 					}
 					UE_LOG( LogNativeCodeGenCommandline, Warning, TEXT("'%s' doesn't appear to be a valid manifest file name/path, defaulting to: '%s'."), 
-						*Value,	*FBlueprintNativeCodeGenManifest::GetDefaultFilename() );
+						*Value,	*FBlueprintNativeCodeGenPaths::GetDefaultManifestPath() );
 				}
 			}
 		}
@@ -212,6 +237,5 @@ FNativeCodeGenCommandlineParams::FNativeCodeGenCommandlineParams(const TArray<FS
 		{
 			OutputDir = FPaths::Combine(*OutputDir, *PluginName);
 		}
-	}	
+	}
 }
-
diff --git a/Engine/Source/Developer/BlueprintNativeCodeGen/Private/NativeCodeGenCommandlineParams.h b/Engine/Source/Developer/BlueprintNativeCodeGen/Private/NativeCodeGenCommandlineParams.h
--- a/Engine/Source/Developer/BlueprintNativeCodeGen/Private/NativeCodeGenCommandlineParams.h
+++ b/Engine/Source/Developer/BlueprintNativeCodeGen/Private/NativeCodeGenCommandlineParams.h
@@ -2,6 +2,24 @@
 
 #pragma once
 
+#include "NativeCodeGenCommandlineParams.generated.h"
+
+/*******************************************************************************
+ * EDependencyConversionRule
+ ******************************************************************************/
+ 
+UENUM()
+enum class EDependencyConversionRule
+{
+	None,
+	RequiredOnly,
+	All,
+};
+
+/*******************************************************************************
+ * FNativeCodeGenCommandlineParams
+ ******************************************************************************/
+
 /**  */
 struct FNativeCodeGenCommandlineParams
 {
@@ -20,5 +38,6 @@ public:
 	FString			ManifestFilePath;
 	bool            bWipeRequested : 1;
 	bool            bPreviewRequested : 1;
+	EDependencyConversionRule DependencyConversionRule;
 };
 
diff --git a/Engine/Source/Developer/BlueprintNativeCodeGen/Public/BlueprintNativeCodeGenModule.h b/Engine/Source/Developer/BlueprintNativeCodeGen/Public/BlueprintNativeCodeGenModule.h
--- a/Engine/Source/Developer/BlueprintNativeCodeGen/Public/BlueprintNativeCodeGenModule.h
+++ b/Engine/Source/Developer/BlueprintNativeCodeGen/Public/BlueprintNativeCodeGenModule.h
@@ -2,9 +2,77 @@
 
 #pragma once
 
+#include "ModuleInterface.h"
+#include "Paths.h"			// for GetCleanFilename()
+#include "ModuleManager.h"	// for GetModuleChecked<>
+#include "Blueprint/BlueprintSupport.h"
+
+#include "BlueprintNativeCodeGenModule.generated.h"
+
+struct BLUEPRINTNATIVECODEGEN_API FCookCommandParams
+{
+	FCookCommandParams(const FString& CommandLine);
+
+	bool    bRunConversion;
+	FString ManifestFilePath;
+	FString OutputPath;
+	FString PluginName;
+
+	TArray<FString> ToConversionParams() const;
+};
+
+/*******************************************************************************
+* UBlueprintNativeCodeGenConfig
+******************************************************************************/
+UCLASS(config = Editor)
+class BLUEPRINTNATIVECODEGEN_API UBlueprintNativeCodeGenConfig : public UObject
+{
+	GENERATED_UCLASS_BODY()
+
+public:
+	UPROPERTY(globalconfig)
+	TArray<FString> PackagesToNeverConvert;
+
+	UPROPERTY(globalconfig)
+	TArray<FString> ExcludedAssetTypes;
+
+	UPROPERTY(globalconfig)
+	TArray<FString> ExcludedBlueprintTypes;
+};
+
+
 /** 
  * 
  */
-class IBlueprintNativeCodeGenModuleInterface : public IModuleInterface
+class IBlueprintNativeCodeGenModule : public IModuleInterface
 {
+public:
+	/**
+	 * Wrapper function that retrieves the interface to this module from the 
+	 * module-manager (so we can keep dependent code free of hardcoded strings,
+	 * used to lookup this module by name).
+	 * 
+	 * @return A reference to this module interface (a singleton).
+	 */
+	static IBlueprintNativeCodeGenModule& Get()
+	{
+		return FModuleManager::LoadModuleChecked<IBlueprintNativeCodeGenModule>(GetModuleName());
+	}
+	
+	/**
+	 * Creates a centralized point where the name of this module is supplied 
+	 * from (so we can avoid littering code with hardcoded strings that 
+	 * all reference this module - in case we want to rename it).
+	 * 
+	 * @return The name of the module that this file is part of.
+	 */
+	static FName GetModuleName()
+	{
+		return TEXT("BlueprintNativeCodeGen");
+	}
+
+	virtual void Convert(UPackage* Package, EReplacementResult ReplacementType) = 0;
+	virtual void SaveManifest(const TCHAR* Filename) = 0;
+	virtual void MergeManifest(const TCHAR* Filename) = 0;
+	virtual void FinalizeManifest() = 0;
 };
diff --git a/Engine/Source/Developer/BlueprintProfiler/BlueprintProfiler.Build.cs b/Engine/Source/Developer/BlueprintProfiler/BlueprintProfiler.Build.cs
new file mode 100644
--- /dev/null
+++ b/Engine/Source/Developer/BlueprintProfiler/BlueprintProfiler.Build.cs
@@ -0,0 +1,26 @@
+// Copyright 1998-2015 Epic Games, Inc. All Rights Reserved.
+
+using UnrealBuildTool;
+
+public class BlueprintProfiler : ModuleRules
+{
+    public BlueprintProfiler(TargetInfo Target)
+    {
+		PrivateIncludePaths.Add("Private");
+
+        PrivateDependencyModuleNames.AddRange(
+            new string[] {
+				"Core",
+				"CoreUObject",
+				"Engine",
+				"EngineSettings",
+			}
+            );
+
+        if (UEBuildConfiguration.bBuildEditor == true)
+        {
+            PrivateDependencyModuleNames.Add("UnrealEd");
+        }
+    }
+}
+
diff --git a/Engine/Source/Developer/BlueprintProfiler/Private/BlueprintProfilerPCH.h b/Engine/Source/Developer/BlueprintProfiler/Private/BlueprintProfilerPCH.h
new file mode 100644
--- /dev/null
+++ b/Engine/Source/Developer/BlueprintProfiler/Private/BlueprintProfilerPCH.h
@@ -0,0 +1,9 @@
+// Copyright 1998-2015 Epic Games, Inc. All Rights Reserved.
+
+#ifndef _INC_BLUEPRINTPROFILER
+#define _INC_BLUEPRINTPROFILER
+
+#include "BlueprintProfilerModule.h"
+#include "Public/BlueprintProfiler.h"
+
+#endif // _INC_BLUEPRINTEDITORPRIVATE
diff --git a/Engine/Source/Developer/BlueprintProfiler/Public/BlueprintProfiler.h b/Engine/Source/Developer/BlueprintProfiler/Public/BlueprintProfiler.h
new file mode 100644
--- /dev/null
+++ b/Engine/Source/Developer/BlueprintProfiler/Public/BlueprintProfiler.h
@@ -0,0 +1,61 @@
+// Copyright 1998-2015 Epic Games, Inc. All Rights Reserved.
+
+#pragma once
+
+#include "BlueprintProfilerModule.h"
+
+//////////////////////////////////////////////////////////////////////////
+// FBlueprintProfiler
+
+class FBlueprintProfiler : public IBlueprintProfilerInterface
+{
+public:
+
+	FBlueprintProfiler();
+	~FBlueprintProfiler();
+
+	// Begin IModuleInterface
+	virtual void StartupModule() override;
+	virtual void ShutdownModule() override;
+	// End IModuleInterface
+
+	// Begin IBlueprintProfilerModule
+	virtual bool IsProfilerEnabled() const override { return bProfilerActive; }
+	virtual bool IsProfilingCaptureActive() const { return bProfilingCaptureActive; }
+	virtual void ToggleProfilingCapture() override;
+	virtual TArray<FBlueprintInstrumentedEvent>& GetProfilingData() override { return ProfilingData; }
+	virtual void InstrumentEvent(const EScriptInstrumentationEvent& Event) override;
+	// End IBlueprintProfilerModule
+
+private:
+
+	/** Reset Profiling Data */
+	virtual void ResetProfilingData();
+
+	/** Registers delegates to recieve profiling events */
+	void RegisterDelegates(bool bEnabled);
+
+	/** Handles pie events to enable/disable profiling captures in the editor */
+	void BeginPIE(bool bIsSimulating);
+	void EndPIE(bool bIsSimulating);
+
+	/** Utility functon to extract script code offset safely */
+	int32 GetScriptCodeOffset(const EScriptInstrumentationEvent& Event) const;
+
+	/** Utility functon to extract object path safely */
+	FString GetCurrentFunctionPath(const EScriptInstrumentationEvent& Event) const;
+
+protected:
+
+	/** Profiler active state */
+	bool bProfilerActive;
+	/** Profiling capture state */
+	bool bProfilingCaptureActive;
+	/** PIE Active */
+	bool bPIEActive;
+	/** Active Object */
+	FKismetInstrumentContext ActiveContext;
+	/** Profiling data */
+	TArray<FBlueprintInstrumentedEvent> ProfilingData;
+	
+};
\ No newline at end of file
diff --git a/Engine/Source/Developer/BlueprintProfiler/Public/BlueprintProfilerModule.h b/Engine/Source/Developer/BlueprintProfiler/Public/BlueprintProfilerModule.h
new file mode 100644
--- /dev/null
+++ b/Engine/Source/Developer/BlueprintProfiler/Public/BlueprintProfilerModule.h
@@ -0,0 +1,30 @@
+// Copyright 1998-2015 Epic Games, Inc. All Rights Reserved.
+
+#pragma once
+
+#include "Script.h"
+#include "BlueprintProfilerSupport.h"
+
+//////////////////////////////////////////////////////////////////////////
+// IBlueprintProfilerInterface
+
+class IBlueprintProfilerInterface : public IModuleInterface
+{
+public:
+
+	/** Returns true if the profiler is enabled */
+	virtual bool IsProfilerEnabled() const { return false; }
+
+	/** Returns true if the profiler is enabled and actively capturing events */
+	virtual bool IsProfilingCaptureActive() const { return false; }
+
+	/** Toggles the profiling event capture state */
+	virtual void ToggleProfilingCapture() {}
+
+	/** Returns the current profiling event data */
+	virtual TArray<FBlueprintInstrumentedEvent>& GetProfilingData() = 0;
+
+	/** Instruments and add a profiling event to the data */
+	virtual void InstrumentEvent(const EScriptInstrumentationEvent& Event) {}
+
+};
diff --git a/Engine/Source/Developer/BlueprintProfiler/Public/BlueprintProfilerSupport.h b/Engine/Source/Developer/BlueprintProfiler/Public/BlueprintProfilerSupport.h
new file mode 100644
--- /dev/null
+++ b/Engine/Source/Developer/BlueprintProfiler/Public/BlueprintProfilerSupport.h
@@ -0,0 +1,78 @@
+// Copyright 1998-2015 Epic Games, Inc. All Rights Reserved.
+
+#pragma once
+
+//////////////////////////////////////////////////////////////////////////
+// FBlueprintInstrumentedEvent
+
+class FBlueprintInstrumentedEvent
+{
+public:
+
+	FBlueprintInstrumentedEvent() {}
+
+	/** Constructs a new script event */
+	FBlueprintInstrumentedEvent(EScriptInstrumentation::Type InEventType, const FString& InPathData, const int32 InCodeOffset)
+		: EventType(InEventType)
+		, PathData(InPathData)
+		, CodeOffset(InCodeOffset)
+		, Time(FPlatformTime::Seconds())
+	{
+	}
+
+	/** Set's the event data */
+	void SetData(EScriptInstrumentation::Type InEventType, const FString& InPathData, const int32 InCodeOffset);
+
+	/** Returns the type of script event */
+	const EScriptInstrumentation::Type GetType() const { return EventType; }
+
+	/** Returns the object path for the object that triggered this event */
+	const FString& GetObjectPath() const { return PathData; }
+
+	/** Returns the script offset at the time this event was generated */
+	int32 GetScriptCodeOffset() const { return CodeOffset; }
+
+	/** Returns the system time when this event was generated */
+	double GetTime() const { return Time; }
+
+	/** Returns if this event represents a change in active class/blueprint */
+	bool IsNewClass() const { return EventType == EScriptInstrumentation::Class; }
+
+	/** Returns if this event represents a change in active instance */
+	bool IsNewInstance() const { return EventType == EScriptInstrumentation::Instance; }
+
+	/** Returns if this event represents a node execution event */
+	bool IsNodeTiming() const {	return EventType == EScriptInstrumentation::NodeEntry || EventType == EScriptInstrumentation::NodeExit;	}
+
+protected:
+	
+	/** The event type */
+	EScriptInstrumentation::Type EventType;
+	/** The object path for the object emitting the event */
+	FString PathData;
+	/** The script code offset when this event was emitted */
+	int32 CodeOffset;
+	/** The recorded system time when this event was emitted */
+	double Time;
+};
+
+//////////////////////////////////////////////////////////////////////////
+// FKismetInstrumentContext
+
+class FKismetInstrumentContext
+{
+public:
+
+	/** Update the current execution context, and emit events to signal changes */
+	void UpdateContext(const UObject* InContextObject, TArray<FBlueprintInstrumentedEvent>& ProfilingData);
+
+	/** Reset the current event context */
+	void ResetContext();
+
+private:
+
+	/** The current class/blueprint context */
+	TWeakObjectPtr<const UClass> ContextClass;
+	/** The current instance context */
+	TWeakObjectPtr<const UObject> ContextObject;
+};
diff --git a/Engine/Source/Developer/BlueprintProfiler/classes/BlueprintProfiler.cpp b/Engine/Source/Developer/BlueprintProfiler/classes/BlueprintProfiler.cpp
new file mode 100644
--- /dev/null
+++ b/Engine/Source/Developer/BlueprintProfiler/classes/BlueprintProfiler.cpp
@@ -0,0 +1,130 @@
+// Copyright 1998-2015 Epic Games, Inc. All Rights Reserved.
+
+#include "Private/BlueprintProfilerPCH.h"
+
+//////////////////////////////////////////////////////////////////////////
+// FBlueprintProfiler
+
+IMPLEMENT_MODULE(FBlueprintProfiler, BlueprintProfiler);
+
+FBlueprintProfiler::FBlueprintProfiler()
+	: bProfilerActive(false)
+	, bProfilingCaptureActive(false)
+	, bPIEActive(false)
+{
+}
+
+FBlueprintProfiler::~FBlueprintProfiler()
+{
+	if (bProfilingCaptureActive)
+	{
+		ToggleProfilingCapture();
+	}
+}
+
+void FBlueprintProfiler::StartupModule()
+{
+	FBlueprintCoreDelegates::OnToggleScriptProfiler.AddRaw(this, &FBlueprintProfiler::RegisterDelegates);
+}
+
+void FBlueprintProfiler::ShutdownModule()
+{
+	FBlueprintCoreDelegates::OnToggleScriptProfiler.RemoveAll(this);
+}
+
+void FBlueprintProfiler::ToggleProfilingCapture()
+{
+	// Toggle profiler state
+	bProfilerActive = !bProfilerActive;
+#if WITH_EDITOR
+	bProfilingCaptureActive = bProfilerActive && bPIEActive;
+#else
+	bProfilingCaptureActive = bProfilerActive;
+#endif // WITH_EDITOR
+	// Broadcast capture state change so delegate consumers can update their state.
+	FBlueprintCoreDelegates::OnToggleScriptProfiler.Broadcast(bProfilerActive);
+}
+
+void FBlueprintProfiler::InstrumentEvent(const EScriptInstrumentationEvent& Event)
+{
+//	if (bProfilingCaptureActive)
+	{
+		// Handle context switching events
+		ActiveContext.UpdateContext(Event.GetContextObject(), ProfilingData);
+		// Add instrumented event
+		const int32 ScriptCodeOffset = GetScriptCodeOffset(Event);
+		const FString CurrentFunctionPath = GetCurrentFunctionPath(Event);
+		ProfilingData.Add(FBlueprintInstrumentedEvent(Event.GetType(), CurrentFunctionPath, ScriptCodeOffset));
+	}
+}
+
+void FBlueprintProfiler::ResetProfilingData()
+{
+	ProfilingData.Reset();
+	ActiveContext.ResetContext();
+}
+
+void FBlueprintProfiler::RegisterDelegates(bool bEnabled)
+{
+	if (bEnabled)
+	{
+#if WITH_EDITOR
+		// Register for PIE begin and end events in the editor
+		FEditorDelegates::BeginPIE.AddRaw(this, &FBlueprintProfiler::BeginPIE);
+		FEditorDelegates::EndPIE.AddRaw(this, &FBlueprintProfiler::EndPIE);
+#endif // WITH_EDITOR
+		ResetProfilingData();
+		// Start consuming profiling events for capture
+		FBlueprintCoreDelegates::OnScriptProfilingEvent.AddRaw(this, &FBlueprintProfiler::InstrumentEvent);
+	}
+	else
+	{
+#if WITH_EDITOR
+		// Unregister for PIE begin and end events in the editor
+		FEditorDelegates::BeginPIE.RemoveAll(this);
+		FEditorDelegates::EndPIE.RemoveAll(this);
+#endif // WITH_EDITOR
+		ResetProfilingData();
+		// Stop consuming profiling events for capture
+		FBlueprintCoreDelegates::OnScriptProfilingEvent.RemoveAll(this);
+	}
+}
+
+void FBlueprintProfiler::BeginPIE(bool bIsSimulating)
+{
+#if WITH_EDITOR
+	bPIEActive = true;
+	ResetProfilingData();
+	bProfilingCaptureActive = bProfilerActive && bPIEActive;
+#endif // WITH_EDITOR
+}
+
+void FBlueprintProfiler::EndPIE(bool bIsSimulating)
+{
+#if WITH_EDITOR
+	bPIEActive = false;
+	bProfilingCaptureActive = bProfilerActive && bPIEActive;
+#endif // WITH_EDITOR
+}
+
+int32 FBlueprintProfiler::GetScriptCodeOffset(const EScriptInstrumentationEvent& Event) const
+{
+	int32 ScriptCodeOffset = -1;
+	if (Event.IsStackFrameValid())
+	{
+		const FFrame& StackFrame = Event.GetStackFrame();
+		ScriptCodeOffset = StackFrame.Code - StackFrame.Node->Script.GetData() - 1;
+	}
+	return ScriptCodeOffset;
+}
+
+FString FBlueprintProfiler::GetCurrentFunctionPath(const EScriptInstrumentationEvent& Event) const
+{
+	FString CurrentFunctionPath;
+	if (Event.IsStackFrameValid())
+	{
+		const FFrame& StackFrame = Event.GetStackFrame();
+		CurrentFunctionPath = StackFrame.Node->GetPathName();
+	}
+	return CurrentFunctionPath;
+}
diff --git a/Engine/Source/Developer/BlueprintProfiler/classes/BlueprintProfilerSupport.cpp b/Engine/Source/Developer/BlueprintProfiler/classes/BlueprintProfilerSupport.cpp
new file mode 100644
--- /dev/null
+++ b/Engine/Source/Developer/BlueprintProfiler/classes/BlueprintProfilerSupport.cpp
@@ -0,0 +1,43 @@
+// Copyright 1998-2015 Epic Games, Inc. All Rights Reserved.
+
+#include "Private/BlueprintProfilerPCH.h"
+
+//////////////////////////////////////////////////////////////////////////
+// FBlueprintInstrumentedEvent
+
+void FBlueprintInstrumentedEvent::SetData(EScriptInstrumentation::Type InEventType, const FString& InPathData, const int32 InCodeOffset)
+{
+	EventType = InEventType;
+	PathData = InPathData;
+	CodeOffset = InCodeOffset;
+	Time = FPlatformTime::Seconds();
+}
+
+//////////////////////////////////////////////////////////////////////////
+// FKismetInstrumentContext
+
+void FKismetInstrumentContext::UpdateContext(const UObject* InContextObject, TArray<FBlueprintInstrumentedEvent>& ProfilingData)
+{
+	// Handle instance context switch
+	if (ContextObject != InContextObject && InContextObject)
+	{
+		// Handle class context switch
+		if (const UClass* NewClass = InContextObject->GetClass())
+		{
+			if (NewClass != ContextClass)
+			{
+				ProfilingData.Add(FBlueprintInstrumentedEvent(EScriptInstrumentation::Class, NewClass->GetPathName(), -1));
+				ContextClass = NewClass;
+			}
+		}
+		// Handle instance change
+		ProfilingData.Add(FBlueprintInstrumentedEvent(EScriptInstrumentation::Instance, InContextObject->GetPathName(), -1));
+		ContextObject = InContextObject;
+	}
+}
+
+void FKismetInstrumentContext::ResetContext()
+{
+	ContextClass.Reset();
+	ContextObject.Reset();
+}
diff --git a/Engine/Source/Editor/BlueprintGraph/Classes/K2Node_CallFunction.h b/Engine/Source/Editor/BlueprintGraph/Classes/K2Node_CallFunction.h
--- a/Engine/Source/Editor/BlueprintGraph/Classes/K2Node_CallFunction.h
+++ b/Engine/Source/Editor/BlueprintGraph/Classes/K2Node_CallFunction.h
@@ -185,10 +185,6 @@ public:
 	static FName GetPaletteIconForFunction(UFunction const* Function, FLinearColor& OutColor);
 
 private: 
-
-	/** Helper function to check if the SelfPin is wired correctly */
-	bool IsSelfPinCompatibleWithBlueprintContext (UEdGraphPin *SelfPin, UBlueprint* BlueprintObj) const;
-
 	/* Looks at function metadata and properties to determine if this node should be using enum to exec expansion */
 	void DetermineWantsEnumToExecExpansion(const UFunction* Function);
 
@@ -209,7 +205,7 @@ private:
 
 protected:
 	/** Helper function to ensure function is called in our context */
-	virtual void EnsureFunctionIsInBlueprint();
+	virtual void FixupSelfMemberContext();
 
 	/** Helper function to find UFunction entries from the skeleton class, use with caution.. */
 	UFunction* GetTargetFunctionFromSkeletonClass() const;
diff --git a/Engine/Source/Editor/BlueprintGraph/Classes/K2Node_Event.h b/Engine/Source/Editor/BlueprintGraph/Classes/K2Node_Event.h
--- a/Engine/Source/Editor/BlueprintGraph/Classes/K2Node_Event.h
+++ b/Engine/Source/Editor/BlueprintGraph/Classes/K2Node_Event.h
@@ -45,6 +45,7 @@ class UK2Node_Event : public UK2Node_EditablePinBase
 	//~ Begin UObject Interface
 	BLUEPRINTGRAPH_API virtual void Serialize(FArchive& Ar) override;
 	BLUEPRINTGRAPH_API virtual void PostLoad() override;
+	BLUEPRINTGRAPH_API virtual void PostDuplicate(bool bDuplicateForPIE) override;
 	//~ End UObject Interface
 
 	//~ Begin UEdGraphNode Interface
@@ -95,6 +96,10 @@ class UK2Node_Event : public UK2Node_EditablePinBase
 
 	/** Helper function to identify if two Event nodes are the same */
 	static BLUEPRINTGRAPH_API bool AreEventNodesIdentical(const UK2Node_Event* InNodeA, const UK2Node_Event* InNodeB);
+
+protected:
+	void FixupEventReference();
+
 private:
 	/** Constructing FText strings can be costly, so we cache the node's tooltip */
 	FNodeTextCache CachedTooltip;
diff --git a/Engine/Source/Editor/BlueprintGraph/Classes/K2Node_FunctionEntry.h b/Engine/Source/Editor/BlueprintGraph/Classes/K2Node_FunctionEntry.h
--- a/Engine/Source/Editor/BlueprintGraph/Classes/K2Node_FunctionEntry.h
+++ b/Engine/Source/Editor/BlueprintGraph/Classes/K2Node_FunctionEntry.h
@@ -64,7 +64,7 @@ class UK2Node_FunctionEntry : public UK2Node_FunctionTerminator
 	BLUEPRINTGRAPH_API int32 GetFunctionFlags() const;
 
 	/** Retrieves the extra flags set on this node. */
-	BLUEPRINTGRAPH_API int32 GetExtraFlags(struct FKismetFunctionContext& CompilerContext) const
+	BLUEPRINTGRAPH_API int32 GetExtraFlags() const
 	{
 		return ExtraFlags;
 	}
@@ -75,6 +75,16 @@ class UK2Node_FunctionEntry : public UK2Node_FunctionTerminator
 		ExtraFlags = (InFlags & ~FUNC_Native);
 	}
 
+	BLUEPRINTGRAPH_API void AddExtraFlags(int32 InFlags)
+	{
+		ExtraFlags |= InFlags;
+	}
+
+	BLUEPRINTGRAPH_API void ClearExtraFlags(int32 InFlags)
+	{
+		ExtraFlags &= ~InFlags;
+	}
+
 protected:
 	/** Any extra flags that the function may need */
 	UPROPERTY()
diff --git a/Engine/Source/Editor/BlueprintGraph/Classes/K2Node_GetClassDefaults.h b/Engine/Source/Editor/BlueprintGraph/Classes/K2Node_GetClassDefaults.h
--- a/Engine/Source/Editor/BlueprintGraph/Classes/K2Node_GetClassDefaults.h
+++ b/Engine/Source/Editor/BlueprintGraph/Classes/K2Node_GetClassDefaults.h
@@ -19,6 +19,7 @@ class UK2Node_GetClassDefaults : public UK2Node
 	virtual void PostPlacedNewNode() override;
 	virtual void PinDefaultValueChanged(UEdGraphPin* Pin) override;
 	virtual void PinConnectionListChanged(UEdGraphPin* Pin) override;
+	virtual void ValidateNodeDuringCompilation(class FCompilerResultsLog& MessageLog) const override;
 	//~ End UEdGraphNode Interface
 
 	//~ Begin UK2Node Interface
@@ -71,4 +72,8 @@ private:
 	/** Output pin visibility control */
 	UPROPERTY(EditAnywhere, Category=PinOptions, EditFixedSize)
 	TArray<FOptionalPinFromProperty> ShowPinForProperties;
+
+	/** Whether or not to exclude object array properties */
+	UPROPERTY()
+	bool bExcludeObjectArrays;
 };
diff --git a/Engine/Source/Editor/BlueprintGraph/Classes/K2Node_MacroInstance.h b/Engine/Source/Editor/BlueprintGraph/Classes/K2Node_MacroInstance.h
--- a/Engine/Source/Editor/BlueprintGraph/Classes/K2Node_MacroInstance.h
+++ b/Engine/Source/Editor/BlueprintGraph/Classes/K2Node_MacroInstance.h
@@ -42,7 +42,6 @@ public:
 	virtual FText GetNodeTitle(ENodeTitleType::Type TitleType) const override;
 	virtual bool CanUserDeleteNode() const override { return true; }
 	void GetContextMenuActions(const FGraphNodeContextMenuBuilder& Context) const override;
-	virtual void ReallocatePinsDuringReconstruction(TArray<UEdGraphPin*>& OldPins) override;
 	virtual void NodeConnectionListChanged() override;
 	virtual FString GetDocumentationLink() const override;
 	virtual FString GetDocumentationExcerptName() const override;
@@ -70,6 +69,10 @@ public:
 	virtual bool CanCreateUserDefinedPin(const FEdGraphPinType& InPinType, EEdGraphPinDirection InDesiredDirection, FText& OutErrorMessage) override { return false; }
 	//~ End UK2Node_EditablePinBase Interface
 
+	// Begin UK2Node_Tunnel interface
+	virtual void PostFixupAllWildcardPins(bool bInAllWildcardPinsUnlinked) override;
+	// End UK2Node_Tunnel interface
+
 	void SetMacroGraph(UEdGraph* Graph) { MacroGraphReference.SetGraph(Graph); }
 	UEdGraph* GetMacroGraph() const { return MacroGraphReference.GetGraph(); }
 	UBlueprint* GetSourceBlueprint() const { return MacroGraphReference.GetBlueprint(); }
diff --git a/Engine/Source/Editor/BlueprintGraph/Classes/K2Node_Message.h b/Engine/Source/Editor/BlueprintGraph/Classes/K2Node_Message.h
--- a/Engine/Source/Editor/BlueprintGraph/Classes/K2Node_Message.h
+++ b/Engine/Source/Editor/BlueprintGraph/Classes/K2Node_Message.h
@@ -17,7 +17,7 @@ class UK2Node_Message : public UK2Node_CallFunction
 	virtual void AllocateDefaultPins() override;
 	//~ End UEdGraphNode Interface.
 
-	//~ Begin K2Node Interface.
+	//~ Begin K2Node Interface.t
 	virtual bool IsNodeSafeToIgnore() const override { return true; }
 	virtual class FNodeHandlingFunctor* CreateNodeHandler(class FKismetCompilerContext& CompilerContext) const override;
 	virtual void ExpandNode(class FKismetCompilerContext& CompilerContext, UEdGraph* SourceGraph) override;
@@ -40,7 +40,7 @@ protected:
 	 */
 	void ExpandLevelStreamingHandlers(class FKismetCompilerContext& InCompilerContext, UEdGraph* InSourceGraph, class UEdGraphPin* InStartingExecPin, class UEdGraphPin* InMessageSelfPin, class UK2Node_DynamicCast* InCastToInterfaceNode);
 
-	virtual void EnsureFunctionIsInBlueprint() override;
+	virtual void FixupSelfMemberContext() override;
 	//~ End K2Node_CallFunction Interface.
 
 	/** Constructing FText strings can be costly, so we cache the node's title */
diff --git a/Engine/Source/Editor/BlueprintGraph/Classes/K2Node_Tunnel.h b/Engine/Source/Editor/BlueprintGraph/Classes/K2Node_Tunnel.h
--- a/Engine/Source/Editor/BlueprintGraph/Classes/K2Node_Tunnel.h
+++ b/Engine/Source/Editor/BlueprintGraph/Classes/K2Node_Tunnel.h
@@ -42,6 +42,7 @@ class UK2Node_Tunnel : public UK2Node_EditablePinBase
 	BLUEPRINTGRAPH_API virtual bool CanDuplicateNode() const override;
 	virtual UObject* GetJumpTargetForDoubleClick() const override;
 	virtual FString CreateUniquePinName(FString SourcePinName) const override;
+	virtual void ReallocatePinsDuringReconstruction(TArray<UEdGraphPin*>& OldPins) override;
 	//~ End UEdGraphNode Interface
 
 	//~ Begin UK2Node Interface.
@@ -58,6 +59,14 @@ class UK2Node_Tunnel : public UK2Node_EditablePinBase
 	virtual bool CanUseRefParams() const override { return true; }
 	virtual bool CanCreateUserDefinedPin(const FEdGraphPinType& InPinType, EEdGraphPinDirection InDesiredDirection, FText& OutErrorMessage) override;
 	//~ End UK2Node_EditablePinBase Interface
+
+protected:
+	/**
+	* Handles any work needed to be done after fixing up all wildcard pins during reconstruction
+	*
+	* @param bInAllWildcardPinsUnlinked	TRUE if all wildcard pins were unlinked
+	*/
+	virtual void PostFixupAllWildcardPins(bool bInAllWildcardPinsUnlinked) {}
 public:
 	// The input pins of this tunnel go to the output pins of InputSinkNode (can be NULL).
 	BLUEPRINTGRAPH_API virtual UK2Node_Tunnel* GetInputSink() const;
diff --git a/Engine/Source/Editor/BlueprintGraph/Private/CallFunctionHandler.cpp b/Engine/Source/Editor/BlueprintGraph/Private/CallFunctionHandler.cpp
--- a/Engine/Source/Editor/BlueprintGraph/Private/CallFunctionHandler.cpp
+++ b/Engine/Source/Editor/BlueprintGraph/Private/CallFunctionHandler.cpp
@@ -230,6 +230,7 @@ void FKCHandler_CallFunction::CreateFunctionCallStatement(FKismetFunctionContext
 									ClassTerm->bIsLiteral = true;
 									ClassTerm->Source     = Node;
 									ClassTerm->ObjectLiteral = InterfaceClass;
+									ClassTerm->Type.PinCategory = CompilerContext.GetSchema()->PC_Class;
 
 									// insert a cast op before a call to the function (and replace
 									// the param with the result from the cast)
diff --git a/Engine/Source/Editor/BlueprintGraph/Private/EdGraphSchema_K2.cpp b/Engine/Source/Editor/BlueprintGraph/Private/EdGraphSchema_K2.cpp
--- a/Engine/Source/Editor/BlueprintGraph/Private/EdGraphSchema_K2.cpp
+++ b/Engine/Source/Editor/BlueprintGraph/Private/EdGraphSchema_K2.cpp
@@ -942,7 +942,7 @@ void UEdGraphSchema_K2::AddExtraFunctionFlags(const UEdGraph* CurrentGraph, int3
 	{
 		if (UK2Node_FunctionEntry* Node = Cast<UK2Node_FunctionEntry>(*It))
 		{
-			Node->SetExtraFlags(Node->GetFunctionFlags() | ExtraFlags);
+			Node->AddExtraFlags(ExtraFlags);
 		}
 	}
 }
@@ -1642,7 +1642,7 @@ void UEdGraphSchema_K2::OnCreateNonExistentLocalVariable( UK2Node_Variable* Vari
 				{
 					if (VariableNode->VariableReference.IsSameReference(OldReference))
 					{
-						VariableNode->VariableReference.SetLocalMember(VarName, Variable->GetGraph()->GetName(), LocalVarGuid);
+						VariableNode->VariableReference.SetLocalMember(VarName, FBlueprintEditorUtils::GetTopLevelGraph(Variable->GetGraph())->GetName(), LocalVarGuid);
 						VariableNode->ReconstructNode();
 					}
 				}
@@ -3152,7 +3152,7 @@ void UEdGraphSchema_K2::CreateFunctionGraphTerminators(UEdGraph& Graph, UClass*
 	if (InterfaceToImplement)
 	{
 		// Add modifier flags from the declaration
-		EntryNode->SetExtraFlags(EntryNode->GetFunctionFlags() | (InterfaceToImplement->FunctionFlags & (FUNC_Const | FUNC_Static | FUNC_BlueprintPure)));
+		EntryNode->AddExtraFlags(InterfaceToImplement->FunctionFlags & (FUNC_Const | FUNC_Static | FUNC_BlueprintPure));
 
 		UK2Node* NextNode = EntryNode;
 		UEdGraphPin* NextExec = FindExecutionPin(*EntryNode, EGPD_Output);
diff --git a/Engine/Source/Editor/BlueprintGraph/Private/K2Node_BaseAsyncTask.cpp b/Engine/Source/Editor/BlueprintGraph/Private/K2Node_BaseAsyncTask.cpp
--- a/Engine/Source/Editor/BlueprintGraph/Private/K2Node_BaseAsyncTask.cpp
+++ b/Engine/Source/Editor/BlueprintGraph/Private/K2Node_BaseAsyncTask.cpp
@@ -225,7 +225,7 @@ bool UK2Node_BaseAsyncTask::FBaseAsyncTaskHelper::HandleDelegateImplementation(
 		bIsErrorFree &= Schema->TryCreateConnection(AddDelegateNode->FindPinChecked(Schema->PN_Self), ProxyObjectPin);
 		bIsErrorFree &= Schema->TryCreateConnection(InOutLastThenPin, AddDelegateNode->FindPinChecked(Schema->PN_Execute));
 		InOutLastThenPin = AddDelegateNode->FindPinChecked(Schema->PN_Then);
-		CurrentCENode->CustomFunctionName = *FString::Printf(TEXT("%s_%s"), *CurrentProperty->GetName(), *CurrentCENode->NodeGuid.ToString());
+		CurrentCENode->CustomFunctionName = *FString::Printf(TEXT("%s_%s"), *CurrentProperty->GetName(), *CompilerContext.GetGuid(CurrentNode));
 		CurrentCENode->AllocateDefaultPins();
 
 		bIsErrorFree &= FBaseAsyncTaskHelper::CreateDelegateForNewFunction(AddDelegateNode->GetDelegatePin(), CurrentCENode->GetFunctionName(), CurrentNode, SourceGraph, CompilerContext);
diff --git a/Engine/Source/Editor/BlueprintGraph/Private/K2Node_CallFunction.cpp b/Engine/Source/Editor/BlueprintGraph/Private/K2Node_CallFunction.cpp
--- a/Engine/Source/Editor/BlueprintGraph/Private/K2Node_CallFunction.cpp
+++ b/Engine/Source/Editor/BlueprintGraph/Private/K2Node_CallFunction.cpp
@@ -1,6 +1,5 @@
 // Copyright 1998-2015 Epic Games, Inc. All Rights Reserved.
 
-
 #include "BlueprintGraphPrivatePCH.h"
 
 #include "CompilerResultsLog.h"
@@ -1624,60 +1623,60 @@ void UK2Node_CallFunction::GetRedirectPinNames(const UEdGraphPin& Pin, TArray<FS
 	}
 }
 
-bool UK2Node_CallFunction::IsSelfPinCompatibleWithBlueprintContext(UEdGraphPin *SelfPin, UBlueprint* BlueprintObj) const
+void UK2Node_CallFunction::FixupSelfMemberContext()
 {
-	check(BlueprintObj);
-
-	UClass* FunctionClass = FunctionReference.GetMemberParentClass(GetBlueprintClassFromNode());
-
-	bool bIsCompatible = (SelfPin != NULL) ? SelfPin->bHidden : true;
-	if (!bIsCompatible && (BlueprintObj->GeneratedClass != NULL))
-	{
-		bIsCompatible |= BlueprintObj->GeneratedClass->IsChildOf(FunctionClass);
-	}
-
-	if (!bIsCompatible && (BlueprintObj->SkeletonGeneratedClass != NULL))
+	UBlueprint* Blueprint = FBlueprintEditorUtils::FindBlueprintForNode(this);
+	auto IsBlueprintOfType = [Blueprint](UClass* ClassType)->bool
 	{
-		bIsCompatible |= BlueprintObj->SkeletonGeneratedClass->IsChildOf(FunctionClass);
-	}
-	return bIsCompatible;
-}
+		bool bIsChildOf  = Blueprint && (Blueprint->GeneratedClass != nullptr) && Blueprint->GeneratedClass->IsChildOf(ClassType);
+		if (!bIsChildOf && Blueprint && (Blueprint->SkeletonGeneratedClass))
+		{
+			bIsChildOf = Blueprint->SkeletonGeneratedClass->IsChildOf(ClassType);
+		}
+		return bIsChildOf;
+	};
+	
 
-void UK2Node_CallFunction::EnsureFunctionIsInBlueprint()
-{
-	// Do not mess with the function if there are pins connected to the target pin
-	UEdGraphPin* SelfPin = GetDefault<UEdGraphSchema_K2>()->FindSelfPin(*this, EGPD_Input);
-	if (SelfPin && SelfPin->LinkedTo.Num() == 0)
+	UClass* MemberClass = FunctionReference.GetMemberParentClass();
+	if (FunctionReference.IsSelfContext())
 	{
-		// Ensure we're calling a function in a context related to our blueprint. If not, 
-		// reassigning the class and then calling ReconstructNodes will re-wire the pins correctly
-		if (UFunction* Function = GetTargetFunction())
+		if (MemberClass == nullptr)
 		{
-			UClass* FunctionOwnerClass = Function->GetOuterUClass();
-			UObject* FunctionGenerator = FunctionOwnerClass ? FunctionOwnerClass->ClassGeneratedBy : NULL;
-
-			// Never change the type if the function is an Interface function type, this only occurs when 
-			// the function is the interface function and the self pin will be PC_Interface type
-			if (!FunctionOwnerClass->IsChildOf(UInterface::StaticClass()))
+			// the self pin may have type information stored on it
+			if (UEdGraphPin* SelfPin = GetDefault<UEdGraphSchema_K2>()->FindSelfPin(*this, EGPD_Input))
 			{
-				// If function is generated from a blueprint object then dbl check self pin compatibility
-				if (FunctionGenerator != NULL)
-				{
-					UBlueprint* BlueprintObj = FBlueprintEditorUtils::FindBlueprintForNode(this);
-					if ((BlueprintObj != NULL) && !IsSelfPinCompatibleWithBlueprintContext(SelfPin, BlueprintObj))
-					{
-						FunctionReference.SetSelfMember(Function->GetFName());
-					}
-				}
+				MemberClass = Cast<UClass>(SelfPin->PinType.PinSubCategoryObject.Get());
+			}
+		}
+		// if we happened to retain the ParentClass for a self reference 
+		// (unlikely), then we know where this node came from... let's keep it
+		// referencing that function
+		if (MemberClass != nullptr)
+		{
+			if (!IsBlueprintOfType(MemberClass))
+			{
+				FunctionReference.SetExternalMember(FunctionReference.GetMemberName(), MemberClass);
 			}
 		}
+		// else, there is nothing we can do... if there is an function matching 
+		// the member name in this Blueprint, then it will reference that 
+		// function (even if it came from a different Blueprint, one with an 
+		// identically named function)... if there is no function matching this 
+		// reference, then the node will produce an error later during compilation
+	}
+	else if (ensure(MemberClass != nullptr))
+	{
+		if (IsBlueprintOfType(MemberClass))
+		{
+			FunctionReference.SetSelfMember(FunctionReference.GetMemberName());
+		}
 	}
 }
 
 void UK2Node_CallFunction::PostPasteNode()
 {
 	Super::PostPasteNode();
-	EnsureFunctionIsInBlueprint();
+	FixupSelfMemberContext();
 
 	UFunction* Function = GetTargetFunction();
 	if(Function != NULL)
@@ -1714,8 +1713,7 @@ void UK2Node_CallFunction::PostDuplicate(bool bDuplicateForPIE)
 	Super::PostDuplicate(bDuplicateForPIE);
 	if (!bDuplicateForPIE && (!this->HasAnyFlags(RF_Transient)))
 	{
-		FunctionReference.InvalidateScope();
-		EnsureFunctionIsInBlueprint();
+		FixupSelfMemberContext();
 	}
 }
 
diff --git a/Engine/Source/Editor/BlueprintGraph/Private/K2Node_Event.cpp b/Engine/Source/Editor/BlueprintGraph/Private/K2Node_Event.cpp
--- a/Engine/Source/Editor/BlueprintGraph/Private/K2Node_Event.cpp
+++ b/Engine/Source/Editor/BlueprintGraph/Private/K2Node_Event.cpp
@@ -1,6 +1,5 @@
 // Copyright 1998-2015 Epic Games, Inc. All Rights Reserved.
 
-
 #include "BlueprintGraphPrivatePCH.h"
 #include "CompilerResultsLog.h"
 #include "KismetCompiler.h"
@@ -50,6 +49,12 @@ void UK2Node_Event::Serialize(FArchive& Ar)
 		{
 			EventReference.SetExternalMember(EventSignatureName_DEPRECATED, EventSignatureClass_DEPRECATED);
 		}
+
+		// @TODO: Dev-BP=>Main; gate this with a version check once it makes its way into main
+		//if (Ar.UE4Ver() < VER_UE4_OVERRIDDEN_EVENT_REFERENCE_FIXUP)
+		{
+			FixupEventReference();
+		}
 	}
 }
 
@@ -65,6 +70,15 @@ void UK2Node_Event::PostLoad()
 	}
 }
 
+void UK2Node_Event::PostDuplicate(bool bDuplicateForPIE)
+{
+	Super::PostDuplicate(bDuplicateForPIE);
+	if (!bDuplicateForPIE)
+	{
+		FixupEventReference();
+	}
+}
+
 FNodeHandlingFunctor* UK2Node_Event::CreateNodeHandler(FKismetCompilerContext& CompilerContext) const
 {
 	return new FKCHandler_EventEntry(CompilerContext);
@@ -182,6 +196,52 @@ void UK2Node_Event::PostReconstructNode()
 	Super::PostReconstructNode();
 }
 
+
+void UK2Node_Event::FixupEventReference()
+{
+	if (bOverrideFunction && !HasAnyFlags(RF_Transient))
+	{
+		if (!EventReference.IsSelfContext())
+		{
+			UBlueprint* Blueprint = GetBlueprint();
+			UClass* BlueprintType = (Blueprint != nullptr) ? Blueprint->SkeletonGeneratedClass : nullptr;
+
+			UClass* ParentType = EventReference.GetMemberParentClass();
+			if ((BlueprintType != nullptr) && ( (ParentType == nullptr) || !(BlueprintType->IsChildOf(ParentType) || BlueprintType->ImplementsInterface(ParentType)) ))
+			{
+				FName EventName = EventReference.GetMemberName();
+
+				const UFunction* OverriddenFunc = BlueprintType->FindFunctionByName(EventName);
+				while (OverriddenFunc != nullptr)
+				{
+					if (UFunction* SuperFunc = OverriddenFunc->GetSuperFunction())
+					{
+						OverriddenFunc = SuperFunc;
+					}
+					else
+					{
+						break;
+					}
+				}
+
+				if (OverriddenFunc != nullptr)
+				{
+					UClass* SuperClass = OverriddenFunc->GetOwnerClass();
+					if (UBlueprint* SuperBlueprint = Cast<UBlueprint>(SuperClass->ClassGeneratedBy))
+					{
+						SuperClass = SuperBlueprint->GeneratedClass;
+					}
+
+					if (SuperClass != nullptr)
+					{
+						EventReference.SetExternalMember(EventName, SuperClass);
+					}
+				}
+			}
+		}
+	}
+}
+
 void UK2Node_Event::UpdateDelegatePin(bool bSilent)
 {
 	UEdGraphPin* Pin = FindPinChecked(DelegateOutputName);
diff --git a/Engine/Source/Editor/BlueprintGraph/Private/K2Node_ExecutionSequence.cpp b/Engine/Source/Editor/BlueprintGraph/Private/K2Node_ExecutionSequence.cpp
--- a/Engine/Source/Editor/BlueprintGraph/Private/K2Node_ExecutionSequence.cpp
+++ b/Engine/Source/Editor/BlueprintGraph/Private/K2Node_ExecutionSequence.cpp
@@ -55,7 +55,7 @@ public:
 		// Process the pins, if there are any valid entries
 		if (OutputPins.Num() > 0)
 		{
-			if (Context.bCreateDebugData && (OutputPins.Num() > 1))
+			if (Context.IsDebuggingOrInstrumentationRequired() && (OutputPins.Num() > 1))
 			{
 				const FString NodeComment = Node->NodeComment.IsEmpty() ? Node->GetName() : Node->NodeComment;
 
@@ -79,7 +79,7 @@ public:
 					{
 						// Emit a debug site
 						FBlueprintCompiledStatement& DebugSiteAndJumpTarget = Context.AppendStatementForNode(Node);
-						DebugSiteAndJumpTarget.Type = KCST_DebugSite;
+						DebugSiteAndJumpTarget.Type = Context.GetBreakpointType();
 						DebugSiteAndJumpTarget.Comment = NodeComment;
 						DebugSiteAndJumpTarget.bIsJumpTarget = true;
 
diff --git a/Engine/Source/Editor/BlueprintGraph/Private/K2Node_FunctionEntry.cpp b/Engine/Source/Editor/BlueprintGraph/Private/K2Node_FunctionEntry.cpp
--- a/Engine/Source/Editor/BlueprintGraph/Private/K2Node_FunctionEntry.cpp
+++ b/Engine/Source/Editor/BlueprintGraph/Private/K2Node_FunctionEntry.cpp
@@ -127,8 +127,22 @@ void UK2Node_FunctionEntry::Serialize(FArchive& Ar)
 			bEnforceConstCorrectness = false;
 		}
 
-		// @TODO: gate this with a version check once it makes its way into main
-		ExtraFlags &= ~FUNC_Native;
+		// @TODO: Dev-BP=>Main; gate this with a version check once it makes its way into main
+		//if (Ar.UE4Ver() < VER_UE4_CLEAN_BLUEPRINT_FUNC_FLAGS))
+		{
+			// Flags we explicitly use ExtraFlags for (at the time this fix was made):
+			//     FUNC_Public, FUNC_Protected, FUNC_Private, 
+			//     FUNC_Static, FUNC_Const,
+			//     FUNC_BlueprintPure, FUNC_BlueprintCallable, FUNC_BlueprintEvent, FUNC_BlueprintAuthorityOnly,
+			//     FUNC_Net, FUNC_NetMulticast, FUNC_NetServer, FUNC_NetClient, FUNC_NetReliable
+			// 
+			// FUNC_Exec, FUNC_Event, & FUNC_BlueprintCosmetic are all inherited 
+			// in FKismetCompilerContext::PrecompileFunction()
+			static const uint32 InvalidExtraFlagsMask = FUNC_Final | FUNC_RequiredAPI | FUNC_BlueprintCosmetic |
+				FUNC_NetRequest | FUNC_Exec | FUNC_Native | FUNC_Event | FUNC_NetResponse | FUNC_MulticastDelegate |
+				FUNC_Delegate | FUNC_HasOutParms | FUNC_HasDefaults | FUNC_DLLImport | FUNC_NetValidate;
+			ExtraFlags &= ~InvalidExtraFlagsMask;
+		}
 	}
 }
 
diff --git a/Engine/Source/Editor/BlueprintGraph/Private/K2Node_FunctionResult.cpp b/Engine/Source/Editor/BlueprintGraph/Private/K2Node_FunctionResult.cpp
--- a/Engine/Source/Editor/BlueprintGraph/Private/K2Node_FunctionResult.cpp
+++ b/Engine/Source/Editor/BlueprintGraph/Private/K2Node_FunctionResult.cpp
@@ -93,10 +93,10 @@ public:
 		{
 			GenerateAssigments(Context, Node);
 
-			if (Context.bCreateDebugData && Node)
+			if (Context.IsDebuggingOrInstrumentationRequired() && Node)
 			{
 				FBlueprintCompiledStatement& TraceStatement = Context.AppendStatementForNode(Node);
-				TraceStatement.Type = KCST_WireTraceSite;
+				TraceStatement.Type = Context.GetWireTraceType();
 				TraceStatement.Comment = Node->NodeComment.IsEmpty() ? Node->GetName() : Node->NodeComment;
 			}
 
diff --git a/Engine/Source/Editor/BlueprintGraph/Private/K2Node_GetClassDefaults.cpp b/Engine/Source/Editor/BlueprintGraph/Private/K2Node_GetClassDefaults.cpp
--- a/Engine/Source/Editor/BlueprintGraph/Private/K2Node_GetClassDefaults.cpp
+++ b/Engine/Source/Editor/BlueprintGraph/Private/K2Node_GetClassDefaults.cpp
@@ -18,10 +18,11 @@ namespace
 	// Optional pin manager subclass.
 	struct FClassDefaultsOptionalPinManager : public FOptionalPinManager
 	{
-		FClassDefaultsOptionalPinManager(UClass* InClass)
+		FClassDefaultsOptionalPinManager(UClass* InClass, bool bExcludeObjectArrays)
 			:FOptionalPinManager()
 		{
 			SrcClass = InClass;
+			bExcludeObjectArrayProperties = bExcludeObjectArrays;
 		}
 
 		virtual void GetRecordDefaults(UProperty* TestProperty, FOptionalPinFromProperty& Record) const override
@@ -34,13 +35,24 @@ namespace
 
 		virtual bool CanTreatPropertyAsOptional(UProperty* TestProperty) const override
 		{
-			// Don't expose object reference properties or anything not marked BlueprintReadOnly/BlueprintReadWrite.
+			// Don't expose anything not marked BlueprintReadOnly/BlueprintReadWrite.
+			if(!TestProperty || !TestProperty->HasAllPropertyFlags(CPF_BlueprintVisible))
+			{
+				return false;
+			}
+			
+			if(UArrayProperty* TestArrayProperty = Cast<UArrayProperty>(TestProperty))
+			{
+				// We only use the Inner type if the flag is set. This is done for backwards-compatibility (some BPs may already rely on the previous behavior when the property value was allowed to be exposed).
+				if(bExcludeObjectArrayProperties && TestArrayProperty->Inner)
+				{
+					TestProperty = TestArrayProperty->Inner;
+				}
+			}
+
+			// Don't expose object properties (except for those containing class objects).
 			// @TODO - Could potentially expose object reference values if/when we have support for 'const' input pins.
-			return TestProperty != nullptr
-				&& !TestProperty->IsA<UClassProperty>()
-				&& !TestProperty->IsA<UObjectProperty>()
-				&& !TestProperty->IsA<UInterfaceProperty>()
-				&& TestProperty->HasAllPropertyFlags(CPF_BlueprintVisible);
+			return !TestProperty->IsA<UObjectProperty>() || TestProperty->IsA<UClassProperty>();
 		}
 
 		virtual void CustomizePinData(UEdGraphPin* Pin, FName SourcePropertyName, int32 ArrayIndex, UProperty* Property = nullptr) const override
@@ -54,6 +66,9 @@ namespace
 	private:
 		// Class type for which optional pins are being managed.
 		UClass* SrcClass;
+
+		// Indicates whether or not object array properties will be excluded.
+		bool bExcludeObjectArrayProperties;
 	};
 
 	// Compilation handler subclass.
@@ -180,6 +195,10 @@ void UK2Node_GetClassDefaults::AllocateDefaultPins()
 
 void UK2Node_GetClassDefaults::PostPlacedNewNode()
 {
+	// Always exclude object array properties for new nodes.
+	// @TODO - Could potentially expose object reference values if/when we have support for 'const' input pins.
+	bExcludeObjectArrays = true;
+
 	if(UEdGraphPin* ClassPin = FindClassPin(Pins))
 	{
 		// Default to the owner BP's generated class for "normal" BPs if this is a new node
@@ -212,6 +231,30 @@ void UK2Node_GetClassDefaults::PinDefaultValueChanged(UEdGraphPin* ChangedPin)
 	}
 }
 
+void UK2Node_GetClassDefaults::ValidateNodeDuringCompilation(class FCompilerResultsLog& MessageLog) const
+{
+	Super::ValidateNodeDuringCompilation(MessageLog);
+
+	if(auto SourceClass = GetInputClass())
+	{
+		for(auto Pin : Pins)
+		{
+			if(Pin && Pin->Direction == EGPD_Output && Pin->LinkedTo.Num() > 0)
+			{
+				// Emit a warning for existing connections to potentially unsafe array property defaults. We do this rather than just implicitly breaking the connection (for compatibility).
+				if(auto ArrayProperty = Cast<UArrayProperty>(SourceClass->FindPropertyByName(FName(*Pin->PinName))))
+				{
+					// Even though array property defaults are copied, the copy could still contain a reference to a non-class object that belongs to the CDO, which would potentially be unsafe to modify.
+					if(ArrayProperty->Inner && ArrayProperty->Inner->IsA<UObjectProperty>() && !ArrayProperty->Inner->IsA<UClassProperty>())
+					{
+						MessageLog.Warning(*LOCTEXT("UnsafeConnectionWarning", "@@ has an unsafe connection to the @@ output pin that is not fully supported at this time. It should be disconnected to avoid potentially corrupting class defaults at runtime. If you need to keep this connection, make sure you're not changing the state of any elements in the array. Also note that if you recreate this node, it will not include this output pin.").ToString(), this, Pin);
+					}
+				}
+			}
+		}
+	}
+}
+
 bool UK2Node_GetClassDefaults::HasExternalDependencies(TArray<class UStruct*>* OptionalOutput) const
 {
 	UClass* SourceClass = GetInputClass();
@@ -333,7 +376,7 @@ UClass* UK2Node_GetClassDefaults::GetInputClass(const UEdGraphPin* FromPin) cons
 void UK2Node_GetClassDefaults::CreateOutputPins(UClass* InClass)
 {
 	// Create the set of output pins through the optional pin manager
-	FClassDefaultsOptionalPinManager OptionalPinManager(InClass);
+	FClassDefaultsOptionalPinManager OptionalPinManager(InClass, bExcludeObjectArrays);
 	OptionalPinManager.RebuildPropertyList(ShowPinForProperties, InClass);
 	OptionalPinManager.CreateVisiblePins(ShowPinForProperties, InClass, EGPD_Output, this);
 
diff --git a/Engine/Source/Editor/BlueprintGraph/Private/K2Node_LoadAsset.cpp b/Engine/Source/Editor/BlueprintGraph/Private/K2Node_LoadAsset.cpp
--- a/Engine/Source/Editor/BlueprintGraph/Private/K2Node_LoadAsset.cpp
+++ b/Engine/Source/Editor/BlueprintGraph/Private/K2Node_LoadAsset.cpp
@@ -94,7 +94,7 @@ void UK2Node_LoadAsset::ExpandNode(class FKismetCompilerContext& CompilerContext
 	// Create OnLoadEvent
 	const FString DelegateOnLoadedParamName(TEXT("OnLoaded"));
 	auto OnLoadEventNode = CompilerContext.SpawnIntermediateNode<UK2Node_CustomEvent>(this, SourceGraph);
-	OnLoadEventNode->CustomFunctionName = *FString::Printf(TEXT("OnLoaded_%s"), *OnLoadEventNode->NodeGuid.ToString());
+	OnLoadEventNode->CustomFunctionName = *FString::Printf(TEXT("OnLoaded_%s"), *CompilerContext.GetGuid(this));
 	OnLoadEventNode->AllocateDefaultPins();
 	{
 		UFunction* LoadAssetFunction = CallLoadAssetNode->GetTargetFunction();
diff --git a/Engine/Source/Editor/BlueprintGraph/Private/K2Node_MacroInstance.cpp b/Engine/Source/Editor/BlueprintGraph/Private/K2Node_MacroInstance.cpp
--- a/Engine/Source/Editor/BlueprintGraph/Private/K2Node_MacroInstance.cpp
+++ b/Engine/Source/Editor/BlueprintGraph/Private/K2Node_MacroInstance.cpp
@@ -402,61 +402,11 @@ bool UK2Node_MacroInstance::CanPasteHere(const UEdGraph* TargetGraph) const
 	return bCanPaste && Super::CanPasteHere(TargetGraph);
 }
 
-void UK2Node_MacroInstance::ReallocatePinsDuringReconstruction(TArray<UEdGraphPin*>& OldPins)
+void UK2Node_MacroInstance::PostFixupAllWildcardPins(bool bInAllWildcardPinsUnlinked)
 {
-	Super::ReallocatePinsDuringReconstruction(OldPins);
-
-	const UEdGraphSchema_K2* const Schema = GetDefault<UEdGraphSchema_K2>();
-
-	// determine if all wildcard pins are unlinked.
-	// if they are, we should revert them all back to wildcard status
-	bool bAllWildcardsAreUnlinked = true;
-	for (auto PinIt = Pins.CreateConstIterator(); PinIt; PinIt++)
-	{
-		// for each of the wildcard pins...
-		UEdGraphPin* const Pin = *PinIt;
-		if ( Pin->PinType.PinCategory == Schema->PC_Wildcard )
-		{
-			// find it in the old pins array (where it might not be a wildcard)
-			// and see if it's unlinked
-			for (auto OldPinIt = OldPins.CreateConstIterator(); OldPinIt; OldPinIt++)
-			{
-				UEdGraphPin const* const OldPin = *OldPinIt;
-				if (OldPin->PinName == Pin->PinName)
-				{
-					if (OldPin->LinkedTo.Num() > 0)
-					{
-						bAllWildcardsAreUnlinked = false;
-						break;
-					}
-				}
-			}
-		}
-	}
-
-	if (bAllWildcardsAreUnlinked == false)
-	{
-		// Copy pin types from old pins for wildcard pins
-		for (auto PinIt = Pins.CreateConstIterator(); PinIt; PinIt++)
-		{
-			UEdGraphPin* const Pin = *PinIt;
-			if ( Pin->PinType.PinCategory == Schema->PC_Wildcard )
-			{
-				// find it in the old pins and copy the type
-				for (auto OldPinIt = OldPins.CreateConstIterator(); OldPinIt; OldPinIt++)
-				{
-					UEdGraphPin const* const OldPin = *OldPinIt;
-					if (OldPin->PinName == Pin->PinName)
-					{
-						Pin->PinType = OldPin->PinType;
-					}
-				}
-			}
-		}
-	}
-	else
+	if (bInAllWildcardPinsUnlinked)
 	{
-		// no type
+		// Reset the type to a wildcard because there are no longer any wildcard pins linked to determine a type with
 		ResolvedWildcardType.ResetToDefaults();
 	}
 }
diff --git a/Engine/Source/Editor/BlueprintGraph/Private/K2Node_Message.cpp b/Engine/Source/Editor/BlueprintGraph/Private/K2Node_Message.cpp
--- a/Engine/Source/Editor/BlueprintGraph/Private/K2Node_Message.cpp
+++ b/Engine/Source/Editor/BlueprintGraph/Private/K2Node_Message.cpp
@@ -73,7 +73,7 @@ UEdGraphPin* UK2Node_Message::CreateSelfPin(const UFunction* Function)
 	return SelfPin;
 }
 
-void UK2Node_Message::EnsureFunctionIsInBlueprint()
+void UK2Node_Message::FixupSelfMemberContext()
 {
 	// Do nothing; the function either exists and works, or doesn't and doesn't
 }
diff --git a/Engine/Source/Editor/BlueprintGraph/Private/K2Node_Tunnel.cpp b/Engine/Source/Editor/BlueprintGraph/Private/K2Node_Tunnel.cpp
--- a/Engine/Source/Editor/BlueprintGraph/Private/K2Node_Tunnel.cpp
+++ b/Engine/Source/Editor/BlueprintGraph/Private/K2Node_Tunnel.cpp
@@ -251,4 +251,59 @@ UObject* UK2Node_Tunnel::GetJumpTargetForDoubleClick() const
 
 	return TargetNode;
 }
+
+void UK2Node_Tunnel::ReallocatePinsDuringReconstruction(TArray<UEdGraphPin*>& OldPins)
+{
+	Super::ReallocatePinsDuringReconstruction(OldPins);
+
+	const UEdGraphSchema_K2* const Schema = GetDefault<UEdGraphSchema_K2>();
+
+	// determine if all wildcard pins are unlinked.
+	// if they are, we should revert them all back to wildcard status
+	bool bAllWildcardsAreUnlinked = true;
+	for (auto PinIt = Pins.CreateConstIterator(); PinIt; PinIt++)
+	{
+		// for each of the wildcard pins...
+		UEdGraphPin* const Pin = *PinIt;
+		if (Pin->PinType.PinCategory == Schema->PC_Wildcard)
+		{
+			// find it in the old pins array (where it might not be a wildcard)
+			// and see if it's unlinked
+			for (auto OldPinIt = OldPins.CreateConstIterator(); OldPinIt; OldPinIt++)
+			{
+				UEdGraphPin const* const OldPin = *OldPinIt;
+				if (OldPin->PinName == Pin->PinName)
+				{
+					if (OldPin->LinkedTo.Num() > 0)
+					{
+						bAllWildcardsAreUnlinked = false;
+						break;
+					}
+				}
+			}
+		}
+	}
+
+	if (bAllWildcardsAreUnlinked == false)
+	{
+		// Copy pin types from old pins for wildcard pins
+		for (auto PinIt = Pins.CreateConstIterator(); PinIt; PinIt++)
+		{
+			UEdGraphPin* const Pin = *PinIt;
+			if (Pin->PinType.PinCategory == Schema->PC_Wildcard)
+			{
+				// find it in the old pins and copy the type
+				for (auto OldPinIt = OldPins.CreateConstIterator(); OldPinIt; OldPinIt++)
+				{
+					UEdGraphPin const* const OldPin = *OldPinIt;
+					if (OldPin->PinName == Pin->PinName)
+					{
+						Pin->PinType = OldPin->PinType;
+					}
+				}
+			}
+		}
+	}
+	PostFixupAllWildcardPins(bAllWildcardsAreUnlinked);
+}
 #undef LOCTEXT_NAMESPACE
\ No newline at end of file
diff --git a/Engine/Source/Editor/ClassViewer/Private/SClassViewer.cpp b/Engine/Source/Editor/ClassViewer/Private/SClassViewer.cpp
--- a/Engine/Source/Editor/ClassViewer/Private/SClassViewer.cpp
+++ b/Engine/Source/Editor/ClassViewer/Private/SClassViewer.cpp
@@ -1760,7 +1760,10 @@ void FClassHierarchy::FindClass(TSharedPtr< FClassViewerNode > InOutClassNode)
 			if (Object->IsA(UBlueprint::StaticClass()))
 			{
 				InOutClassNode->Blueprint = Cast<UBlueprint>(Object);
-				InOutClassNode->Class = Cast<UClass>(InOutClassNode->Blueprint->GeneratedClass);
+				if (InOutClassNode->Blueprint.IsValid())
+				{
+					InOutClassNode->Class = Cast<UClass>(InOutClassNode->Blueprint->GeneratedClass);
+				}
 			}
 			else if (UClass* Class = Cast<UClass>(Object))
 			{
diff --git a/Engine/Source/Editor/ContentBrowser/Private/AssetContextMenu.cpp b/Engine/Source/Editor/ContentBrowser/Private/AssetContextMenu.cpp
--- a/Engine/Source/Editor/ContentBrowser/Private/AssetContextMenu.cpp
+++ b/Engine/Source/Editor/ContentBrowser/Private/AssetContextMenu.cpp
@@ -423,13 +423,17 @@ void FAssetContextMenu::MakeAssetActionsSubMenu(FMenuBuilder& MenuBuilder)
 
 		if (bCanUsePropertyMatrix)
 		{
+			TAttribute<FText>::FGetter DynamicTooltipGetter;
+			DynamicTooltipGetter.BindSP(this, &FAssetContextMenu::GetExecutePropertyMatrixTooltip);
+			TAttribute<FText> DynamicTooltipAttribute = TAttribute<FText>::Create(DynamicTooltipGetter);
+
 			MenuBuilder.AddMenuEntry(
 				LOCTEXT("PropertyMatrix", "Bulk Edit via Property Matrix..."),
-				LOCTEXT("PropertyMatrixTooltip", "Opens the property matrix editor for the selected assets."),
+				DynamicTooltipAttribute,
 				FSlateIcon(),
 				FUIAction(
 				FExecuteAction::CreateSP(this, &FAssetContextMenu::ExecutePropertyMatrix),
-				FCanExecuteAction::CreateSP(this, &FAssetContextMenu::CanExecuteProperties)
+				FCanExecuteAction::CreateSP(this, &FAssetContextMenu::CanExecutePropertyMatrix)
 				)
 				);
 		}
@@ -1103,7 +1107,7 @@ void FAssetContextMenu::ExecuteCreateBlueprintUsing()
 	}
 }
 
-void FAssetContextMenu::GetSelectedAssets(TArray<UObject*>& Assets, bool SkipRedirectors)
+void FAssetContextMenu::GetSelectedAssets(TArray<UObject*>& Assets, bool SkipRedirectors) const
 {
 	for (int32 AssetIdx = 0; AssetIdx < SelectedAssets.Num(); ++AssetIdx)
 	{
@@ -1819,9 +1823,46 @@ bool FAssetContextMenu::CanExecuteProperties() const
 	return bAtLeastOneNonRedirectorSelected;
 }
 
+bool FAssetContextMenu::CanExecutePropertyMatrix(FText& OutErrorMessage) const
+{
+	bool bResult = bAtLeastOneNonRedirectorSelected;
+	if (bAtLeastOneNonRedirectorSelected)
+	{
+		TArray<UObject*> ObjectsForPropertiesMenu;
+		const bool SkipRedirectors = true;
+		GetSelectedAssets(ObjectsForPropertiesMenu, SkipRedirectors);
+
+		// Ensure all Blueprints are valid.
+		for (UObject* Object : ObjectsForPropertiesMenu)
+		{
+			if (UBlueprint* BlueprintObj = Cast<UBlueprint>(Object))
+			{
+				if (BlueprintObj->GeneratedClass == nullptr)
+				{
+					OutErrorMessage = LOCTEXT("InvalidBlueprint", "A selected Blueprint is invalid.");
+					bResult = false;
+					break;
+				}
+			}
+		}
+	}
+	return bResult;
+}
+
 bool FAssetContextMenu::CanExecutePropertyMatrix() const
 {
-	return bAtLeastOneNonRedirectorSelected;
+	FText ErrorMessageDummy;
+	return CanExecutePropertyMatrix(ErrorMessageDummy);
+}
+
+FText FAssetContextMenu::GetExecutePropertyMatrixTooltip() const
+{
+	FText ResultTooltip;
+	if (CanExecutePropertyMatrix(ResultTooltip))
+	{
+		ResultTooltip = LOCTEXT("PropertyMatrixTooltip", "Opens the property matrix editor for the selected assets.");
+	}
+	return ResultTooltip;
 }
 
 bool FAssetContextMenu::CanExecuteDuplicate() const
diff --git a/Engine/Source/Editor/ContentBrowser/Private/AssetContextMenu.h b/Engine/Source/Editor/ContentBrowser/Private/AssetContextMenu.h
--- a/Engine/Source/Editor/ContentBrowser/Private/AssetContextMenu.h
+++ b/Engine/Source/Editor/ContentBrowser/Private/AssetContextMenu.h
@@ -237,6 +237,7 @@ private:
 	bool CanExecuteProperties() const;
 
 	/** Handler to check to see if a property matrix command is allowed */
+	bool CanExecutePropertyMatrix(FText& OutErrorMessage) const;
 	bool CanExecutePropertyMatrix() const;
 
 	/** Handler to check to see if a "Remove from collection" command is allowed */
@@ -311,9 +312,12 @@ private:
 	/** Handles when the Assign chunkID dialog Cancel button is clicked */
 	FReply OnChunkIDAssignCancel(TSharedPtr<SWindow> Window);
 
+	/** Generates tooltip for the Property Matrix menu option */
+	FText GetExecutePropertyMatrixTooltip() const;
+
 private:
 	/** Generates a list of selected assets in the content browser */
-	void GetSelectedAssets(TArray<UObject*>& Assets, bool SkipRedirectors);
+	void GetSelectedAssets(TArray<UObject*>& Assets, bool SkipRedirectors) const;
 
 	TArray<FAssetData> SelectedAssets;
 	FSourcesData SourcesData;
diff --git a/Engine/Source/Editor/DetailCustomizations/Private/AnimMontageSegmentDetails.cpp b/Engine/Source/Editor/DetailCustomizations/Private/AnimMontageSegmentDetails.cpp
--- a/Engine/Source/Editor/DetailCustomizations/Private/AnimMontageSegmentDetails.cpp
+++ b/Engine/Source/Editor/DetailCustomizations/Private/AnimMontageSegmentDetails.cpp
@@ -273,9 +273,9 @@ void SAnimationSegmentViewport::Tick( const FGeometry& AllottedGeometry, const d
 		{
 			Description->SetText(FText::Format( LOCTEXT("Previewing", "Previewing {0}"), FText::FromString(Component->GetPreviewText()) ));
 		}
-		else if (Component->AnimBlueprintGeneratedClass)
+		else if (Component->AnimClass)
 		{
-			Description->SetText(FText::Format( LOCTEXT("Previewing", "Previewing {0}"), FText::FromString(Component->AnimBlueprintGeneratedClass->GetName()) ));
+			Description->SetText(FText::Format( LOCTEXT("Previewing", "Previewing {0}"), FText::FromString(Component->AnimClass->GetName()) ));
 		}
 		else if (Component->SkeletalMesh == NULL)
 		{
diff --git a/Engine/Source/Editor/DetailCustomizations/Private/AnimSequenceDetails.cpp b/Engine/Source/Editor/DetailCustomizations/Private/AnimSequenceDetails.cpp
--- a/Engine/Source/Editor/DetailCustomizations/Private/AnimSequenceDetails.cpp
+++ b/Engine/Source/Editor/DetailCustomizations/Private/AnimSequenceDetails.cpp
@@ -548,9 +548,9 @@ void SAnimationRefPoseViewport::Tick( const FGeometry& AllottedGeometry, const d
 
 			Description->SetText( FText::Format( LOCTEXT( "Previewing", "Previewing {0}" ), FText::FromString( Component->GetPreviewText() ) ) );
 		}
-		else if ( Component->AnimBlueprintGeneratedClass )
+		else if ( Component->AnimClass )
 		{
-			Description->SetText( FText::Format( LOCTEXT( "Previewing", "Previewing {0}" ), FText::FromString( Component->AnimBlueprintGeneratedClass->GetName() ) ) );
+			Description->SetText( FText::Format( LOCTEXT( "Previewing", "Previewing {0}" ), FText::FromString( Component->AnimClass->GetName() ) ) );
 		}
 		else if ( AnimRef && AnimRef->GetSkeleton() != TargetSkeleton )
 		{
diff --git a/Engine/Source/Editor/DetailCustomizations/Private/SkeletalMeshComponentDetails.cpp b/Engine/Source/Editor/DetailCustomizations/Private/SkeletalMeshComponentDetails.cpp
--- a/Engine/Source/Editor/DetailCustomizations/Private/SkeletalMeshComponentDetails.cpp
+++ b/Engine/Source/Editor/DetailCustomizations/Private/SkeletalMeshComponentDetails.cpp
@@ -120,7 +120,7 @@ void FSkeletalMeshComponentDetails::UpdateAnimationCategory( IDetailLayoutBuilde
 	AnimationModeHandle = DetailBuilder.GetProperty(GET_MEMBER_NAME_CHECKED(USkeletalMeshComponent, AnimationMode));
 	check (AnimationModeHandle->IsValidHandle());
 
-	AnimationBlueprintHandle = DetailBuilder.GetProperty(GET_MEMBER_NAME_CHECKED(USkeletalMeshComponent, AnimBlueprintGeneratedClass));
+	AnimationBlueprintHandle = DetailBuilder.GetProperty(GET_MEMBER_NAME_CHECKED(USkeletalMeshComponent, AnimClass));
 	check(AnimationBlueprintHandle->IsValidHandle());
 
 	AnimationCategory.AddProperty(AnimationModeHandle);
