Copying //UE4/Dev-Blueprints to Dev-Main (//UE4/Dev-Main)

==========================
MAJOR FEATURES + CHANGES
==========================

Change 2781504 on 2015/11/25 by [REDACTED_PERSON]

	Guarding against invalid nodes for deferred graph node actions (add, remove, select), by using TWeakObjectPtr instead of raw UEdGraphNode pointers.

Change 2781513 on 2015/11/25 by [REDACTED_PERSON]

	Find-in-Blueprints optimized gathering.

	Size of data has shrunk in the Asset Registry by up to one fifth the old size!

	Performance moderately improved.

	Load and save times of Blueprints increased, less redundant gathering of searchable data.

Change 2781517 on 2015/11/25 by [REDACTED_PERSON]

	Marked FTimerHandle::Handle as a UPROPERTY(transient) so that Blueprints can check the equality of two instances of the structure.

Change 2781804 on 2015/11/26 by [REDACTED_PERSON]

	Changed ConformImplementedEvents.

Change 2781821 on 2015/11/26 by [REDACTED_PERSON]

	This reinstates the blueprint debugging keymaps and adds additional functionality for step over and step out as key maps in the PIE world controls.
	#UEBP-66 - Blueprint debug keymappings
	#UE-16817 - Add step-in, step-over, and run until here functions for breakpoints
	#UE-12481 - The F10 key doesn't work for stepping blueprint debugging
	#Branch UE4
	#Proj GraphEditor, Kismet, UnrealEd, CoreUObject, Slate

Change 2781861 on 2015/11/26 by [REDACTED_PERSON]

	UE-23626 Converted tower defense game - you cannot click to place towers

	CodeGenerator generates overriden exported names for events and functions.

Change 2782798 on 2015/11/30 by [REDACTED_PERSON]

	BP C++ conversion: components from SCS calls AttachTo (with ParentSocket parameter).

Change 2782881 on 2015/11/30 by [REDACTED_PERSON]

	Fixed ensure when promoting function graphs from interfaces during interface removal.

Change 2783041 on 2015/11/30 by [REDACTED_PERSON]

	BP C++ conversion: All variables from Event Graph are listed as class properties.

Change 2783080 on 2015/11/30 by [REDACTED_PERSON]

	Removing an interface function's output parameters will no longer cause Blueprints implementing the function to error.

	Functions expected as event overrides will accept function graph implementations and give a warning informing that it is unexpected.

	All function graphs (interfaces, interface implementations, overrides) can be duplicated. Parent function calls will be removed.

	Duplicating graphs will correct names of objects in child Blueprints.

	Function overrides of interfaces expected as an event can be deleted.

	Duplicating graphs while in PIE is no longer possible.

	When removing an interface, the operation can now be canceled.

Change 2783338 on 2015/11/30 by [REDACTED_PERSON]

	New output pins on function result nodes will properly fill out with valid default values.

	All invalid pins will auto-validate themselves on node reconstruction when opening the Blueprint.

Change 2783742 on 2015/11/30 by [REDACTED_PERSON]

	[UE-15463] Add special-case handling for failed imports of BPGC-owned component archetype objects on level load.

	change summary:
	- modified FLinkerLoad::VerifyImport() to customize the load error messaging for missing component archetype objects

Change 2784652 on 2015/12/01 by [REDACTED_PERSON]

	Fix for crash whilst undoing the creation of a macro and currently displaying the tooltip in the blueprint editor.
	#UE-23955 - Adding a macro graph through MyBlueprint and then calling undo causes a crash updating the macro tooltip.
	#Branch UE4
	#Proj Kismet

Change 2784834 on 2015/12/01 by [REDACTED_PERSON]

	Added functions to convert from string to: Vector, Vector2D, Rotator, Color.

	PR #1795

diff --git a/Engine/Config/BaseEditor.ini b/Engine/Config/BaseEditor.ini
--- a/Engine/Config/BaseEditor.ini
+++ b/Engine/Config/BaseEditor.ini
@@ -463,6 +463,3 @@ AngleUnits=EUnit::Degrees
 SpeedUnits=EUnit::MetersPerSecond
 TemperatureUnits=EUnit::Celsius
 ForceUnits=EUnit::Newtons
-
-[/Script/BlueprintNativeCodeGen.BlueprintNativeCodeGenConfig]
-+ExcludedBlueprintTypes=/Script/Engine.LevelScriptActor
diff --git a/Engine/Source/Developer/BlueprintCompilerCppBackend/Private/BlueprintCompilerCppBackend.cpp b/Engine/Source/Developer/BlueprintCompilerCppBackend/Private/BlueprintCompilerCppBackend.cpp
--- a/Engine/Source/Developer/BlueprintCompilerCppBackend/Private/BlueprintCompilerCppBackend.cpp
+++ b/Engine/Source/Developer/BlueprintCompilerCppBackend/Private/BlueprintCompilerCppBackend.cpp
@@ -92,8 +92,8 @@ void FBlueprintCompilerCppBackend::EmitCallStatment(FEmitterLocalContext& Emitte
 void FBlueprintCompilerCppBackend::EmitAssignmentStatment(FEmitterLocalContext& EmitterContext, FKismetFunctionContext& FunctionContext, FBlueprintCompiledStatement& Statement)
 {
 	check(Statement.LHS && Statement.RHS[0]);
-	FString DestinationExpression = TermToText(EmitterContext, Statement.LHS, false);
-	FString SourceExpression = TermToText(EmitterContext, Statement.RHS[0]);
+	const FString DestinationExpression = TermToText(EmitterContext, Statement.LHS, false, false);
+	const FString SourceExpression = TermToText(EmitterContext, Statement.RHS[0]);
 
 	FSafeContextScopedEmmitter SafeContextScope(EmitterContext, Statement.LHS->Context, *this);
 
@@ -109,7 +109,8 @@ void FBlueprintCompilerCppBackend::EmitCastObjToInterfaceStatement(FEmitterLocal
 	FString ObjectValue = TermToText(EmitterContext, Statement.RHS[1]);
 	FString InterfaceValue = TermToText(EmitterContext, Statement.LHS);
 
-	EmitterContext.AddLine(FString::Printf(TEXT("if ( IsValid(%s) && %s->GetClass()->ImplementsInterface(%s) )"), *ObjectValue, *ObjectValue, *InterfaceClass));
+	// Both here and in UObject::execObjectToInterface IsValid function should be used.
+	EmitterContext.AddLine(FString::Printf(TEXT("if ( %s && %s->GetClass()->ImplementsInterface(%s) )"), *ObjectValue, *ObjectValue, *InterfaceClass));
 	EmitterContext.AddLine(FString::Printf(TEXT("{")));
 	EmitterContext.AddLine(FString::Printf(TEXT("\t%s.SetObject(%s);"), *InterfaceValue, *ObjectValue));
 	EmitterContext.AddLine(FString::Printf(TEXT("\tvoid* IAddress = %s->GetInterfaceAddress(%s);"), *ObjectValue, *InterfaceClass));
@@ -145,7 +146,7 @@ void FBlueprintCompilerCppBackend::EmitCastInterfaceToObjStatement(FEmitterLocal
 {
 	FString ClassToCastTo = TermToText(EmitterContext, Statement.RHS[0]);
 	FString InputInterface = TermToText(EmitterContext, Statement.RHS[1]);
-	FString ResultObject = TermToText(EmitterContext, Statement.LHS);
+	FString ResultObject = TermToText(EmitterContext, Statement.LHS, true, false);
 	FString InputObject = FString::Printf(TEXT("%s.GetObjectRef()"), *InputInterface);
 	EmitterContext.AddLine(FString::Printf(TEXT("%s = Cast<%s>(%s);"),*ResultObject, *ClassToCastTo, *InputObject));
 }
@@ -156,7 +157,7 @@ void FBlueprintCompilerCppBackend::EmitDynamicCastStatement(FEmitterLocalContext
 	check(ClassPtr != nullptr);
 	
 	const FString ObjectValue = TermToText(EmitterContext, Statement.RHS[1]);
-	const FString CastedValue = TermToText(EmitterContext, Statement.LHS);
+	const FString CastedValue = TermToText(EmitterContext, Statement.LHS, true, false);
 
 	auto BPGC = Cast<UBlueprintGeneratedClass>(ClassPtr);
 	if (BPGC && !EmitterContext.Dependencies.WillClassBeConverted(BPGC))
@@ -176,14 +177,14 @@ void FBlueprintCompilerCppBackend::EmitMetaCastStatement(FEmitterLocalContext& E
 {
 	FString DesiredClass = TermToText(EmitterContext, Statement.RHS[0]);
 	FString SourceClass = TermToText(EmitterContext, Statement.RHS[1]);
-	FString Destination = TermToText(EmitterContext, Statement.LHS);
+	FString Destination = TermToText(EmitterContext, Statement.LHS, true, false);
 	EmitterContext.AddLine(FString::Printf(TEXT("%s = DynamicMetaCast(%s, %s);"),*Destination, *DesiredClass, *SourceClass));
 }
 
 void FBlueprintCompilerCppBackend::EmitObjectToBoolStatement(FEmitterLocalContext& EmitterContext, FKismetFunctionContext& FunctionContext, FBlueprintCompiledStatement& Statement)
 {
 	FString ObjectTarget = TermToText(EmitterContext, Statement.RHS[0]);
-	FString DestinationExpression = TermToText(EmitterContext, Statement.LHS);
+	FString DestinationExpression = TermToText(EmitterContext, Statement.LHS, true, false);
 	EmitterContext.AddLine(FString::Printf(TEXT("%s = (%s != nullptr);"), *DestinationExpression, *ObjectTarget));
 }
 
@@ -418,7 +419,8 @@ FString FBlueprintCompilerCppBackend::EmitCallStatmentInner(FEmitterLocalContext
 	const bool bCallOnDifferentObject = Statement.FunctionContext && (Statement.FunctionContext->Name != TEXT("self"));
 	const bool bStaticCall = Statement.FunctionToCall->HasAnyFunctionFlags(FUNC_Static);
 	const bool bUseSafeContext = bCallOnDifferentObject && !bStaticCall;
-	const bool bInterfaceCall = bCallOnDifferentObject && Statement.FunctionContext && (UEdGraphSchema_K2::PC_Interface == Statement.FunctionContext->Type.PinCategory);
+	const bool bAnyInterfaceCall = bCallOnDifferentObject && Statement.FunctionContext && (UEdGraphSchema_K2::PC_Interface == Statement.FunctionContext->Type.PinCategory);
+	const bool bInterfaceCallExecute = bAnyInterfaceCall && Statement.FunctionToCall && Statement.FunctionToCall->HasAnyFunctionFlags(FUNC_Event);
 
 	FString Result;
 	FString CloseCast;
@@ -442,7 +444,7 @@ FString FBlueprintCompilerCppBackend::EmitCallStatmentInner(FEmitterLocalContext
 
 	// Emit object to call the method on
 	const FString FunctionToCallOriginalName = FEmitHelper::GetCppName(FEmitHelper::GetOriginalFunction(Statement.FunctionToCall));
-	if (bInterfaceCall)
+	if (bInterfaceCallExecute)
 	{
 		auto ContextInterfaceClass = CastChecked<UClass>(Statement.FunctionContext->Type.PinSubCategoryObject.Get());
 		ensure(ContextInterfaceClass->IsChildOf<UInterface>());
@@ -493,7 +495,7 @@ FString FBlueprintCompilerCppBackend::EmitCallStatmentInner(FEmitterLocalContext
 		Result += TEXT("(");
 	}
 	const FString ParameterList = EmitMethodInputParameterList(EmitterContext, Statement);
-	if (bInterfaceCall && !ParameterList.IsEmpty())
+	if (bInterfaceCallExecute && !ParameterList.IsEmpty())
 	{
 		Result += TEXT(", ");
 	}
@@ -508,7 +510,7 @@ FString FBlueprintCompilerCppBackend::EmitCallStatmentInner(FEmitterLocalContext
 	return Result;
 }
 
-FString FBlueprintCompilerCppBackend::TermToText(FEmitterLocalContext& EmitterContext, const FBPTerminal* Term, bool bUseSafeContext)
+FString FBlueprintCompilerCppBackend::TermToText(FEmitterLocalContext& EmitterContext, const FBPTerminal* Term, bool bUseSafeContext, bool bGetter)
 {
 	const FString PSC_Self(TEXT("self"));
 	if (Term->bIsLiteral)
@@ -576,15 +578,16 @@ FString FBlueprintCompilerCppBackend::TermToText(FEmitterLocalContext& EmitterCo
 		}
 		else if (Term->AssociatedVarProperty)
 		{
+			const bool bSelfContext = (!Term->Context) || (Term->Context->Name == PSC_Self);
 			const bool bPropertyOfParent = EmitterContext.Dependencies.GetActualStruct()->IsChildOf(Term->AssociatedVarProperty->GetOwnerStruct());
 			const bool bIsAccessible = !Term->AssociatedVarProperty->HasAnyPropertyFlags(CPF_NativeAccessSpecifierPrivate)
-				&& (bPropertyOfParent || !Term->AssociatedVarProperty->HasAnyPropertyFlags(CPF_NativeAccessSpecifierProtected));
+				&& ((bPropertyOfParent && bSelfContext) || !Term->AssociatedVarProperty->HasAnyPropertyFlags(CPF_NativeAccessSpecifierProtected));
 
 			auto MinimalClass = Term->AssociatedVarProperty->GetOwnerClass();
 			auto MinimalBPGC = Cast<UBlueprintGeneratedClass>(MinimalClass);
 			if (MinimalBPGC && !EmitterContext.Dependencies.WillClassBeConverted(MinimalBPGC))
 			{
-				if ((!Term->Context) || (Term->Context->Name == PSC_Self))
+				if (bSelfContext)
 				{
 					ensure(ContextStr.IsEmpty());
 					ContextStr = TEXT("this");
@@ -597,7 +600,7 @@ FString FBlueprintCompilerCppBackend::TermToText(FEmitterLocalContext& EmitterCo
 			}
 			else if (!bIsAccessible)
 			{
-				if ((!Term->Context) || (Term->Context->Name == PSC_Self))
+				if (bSelfContext)
 				{
 					ensure(ContextStr.IsEmpty());
 					ContextStr = TEXT("this");
@@ -606,7 +609,7 @@ FString FBlueprintCompilerCppBackend::TermToText(FEmitterLocalContext& EmitterCo
 			}
 			else
 			{
-				if ((Term->Context != nullptr) && (Term->Context->Name != PSC_Self))
+				if (!bSelfContext)
 				{
 					ResultPath = ContextStr + TEXT("->");
 				}
@@ -619,6 +622,11 @@ FString FBlueprintCompilerCppBackend::TermToText(FEmitterLocalContext& EmitterCo
 			ResultPath = Term->Name;
 		}
 
+		if (Term->Type.bIsWeakPointer && bGetter)
+		{
+			ResultPath += TEXT(".Get()");
+		}
+
 		FString Conditions;
 		if (bUseSafeContext)
 		{
@@ -627,16 +635,11 @@ FString FBlueprintCompilerCppBackend::TermToText(FEmitterLocalContext& EmitterCo
 
 		if (!Conditions.IsEmpty())
 		{
-			FString DefaultValue = FEmitHelper::DefaultValue(EmitterContext, Term->Type);
-			if (Term->Type.bIsArray)
-			{
-				// to fix error C4239 in VS, when the term is passed as non-const reference. 
-				// It seems that this situation happens only for arrays, otherwise this should be a default path.
-				const FString DefaultValueVariable = EmitterContext.GenerateUniqueLocalName();
-				EmitterContext.AddLine(*FString::Printf(TEXT("auto %s = %s;"), *DefaultValueVariable, *DefaultValue));
-				DefaultValue = DefaultValueVariable;
-			}
-			return FString::Printf(TEXT("((%s) ? (%s) : (%s))"), *Conditions, *ResultPath, *DefaultValue);
+			const FString DefaultValueConstructor = FEmitHelper::DefaultValue(EmitterContext, Term->Type);
+			const FString DefaultValueVariable = EmitterContext.GenerateUniqueLocalName();
+			const FString CppType = FEmitHelper::PinTypeToNativeType(Term->Type);
+			EmitterContext.AddLine(*FString::Printf(TEXT("%s %s = %s;"), *CppType, *DefaultValueVariable, *DefaultValueConstructor));
+			return FString::Printf(TEXT("((%s) ? (%s) : (%s))"), *Conditions, *ResultPath, *DefaultValueVariable);
 		}
 		return ResultPath;
 	}
diff --git a/Engine/Source/Developer/BlueprintCompilerCppBackend/Private/BlueprintCompilerCppBackend.h b/Engine/Source/Developer/BlueprintCompilerCppBackend/Private/BlueprintCompilerCppBackend.h
--- a/Engine/Source/Developer/BlueprintCompilerCppBackend/Private/BlueprintCompilerCppBackend.h
+++ b/Engine/Source/Developer/BlueprintCompilerCppBackend/Private/BlueprintCompilerCppBackend.h
@@ -37,5 +37,5 @@ protected:
 	FString EmitCallStatmentInner(FEmitterLocalContext& EmitterContext, FBlueprintCompiledStatement& Statement, bool bInline);
 
 public:
-	FString TermToText(FEmitterLocalContext& EmitterContext, const FBPTerminal* Term, bool bUseSafeContext = true);
+	FString TermToText(FEmitterLocalContext& EmitterContext, const FBPTerminal* Term, bool bUseSafeContext = true, bool bGetter = true);
 };
\ No newline at end of file
diff --git a/Engine/Source/Developer/BlueprintCompilerCppBackend/Private/BlueprintCompilerCppBackendAnim.cpp b/Engine/Source/Developer/BlueprintCompilerCppBackend/Private/BlueprintCompilerCppBackendAnim.cpp
--- a/Engine/Source/Developer/BlueprintCompilerCppBackend/Private/BlueprintCompilerCppBackendAnim.cpp
+++ b/Engine/Source/Developer/BlueprintCompilerCppBackend/Private/BlueprintCompilerCppBackendAnim.cpp
@@ -5,10 +5,13 @@
 #include "Animation/AnimBlueprintGeneratedClass.h"
 #include "Animation/AnimClassData.h"
 
-FString FBackendHelperAnim::AddHeaders(UClass* GeneratedClass)
+void FBackendHelperAnim::AddHeaders(FEmitterLocalContext& Context)
 {
-	auto AnimClass = Cast<UAnimBlueprintGeneratedClass>(GeneratedClass);
-	return AnimClass ? FString(TEXT("#include \"Animation/AnimClassData.h\"")) : FString();
+	if (Cast<UAnimBlueprintGeneratedClass>(Context.GetCurrentlyGeneratedClass()))
+	{
+		Context.Header.AddLine(TEXT("#include \"Animation/AnimClassData.h\""));
+		Context.Body.AddLine(TEXT("#include \"Animation/BlendProfile.h\""));
+	}
 }
 
 void FBackendHelperAnim::CreateAnimClassData(FEmitterLocalContext& Context)
diff --git a/Engine/Source/Developer/BlueprintCompilerCppBackend/Private/BlueprintCompilerCppBackendBase.cpp b/Engine/Source/Developer/BlueprintCompilerCppBackend/Private/BlueprintCompilerCppBackendBase.cpp
--- a/Engine/Source/Developer/BlueprintCompilerCppBackend/Private/BlueprintCompilerCppBackendBase.cpp
+++ b/Engine/Source/Developer/BlueprintCompilerCppBackend/Private/BlueprintCompilerCppBackendBase.cpp
@@ -650,19 +650,21 @@ void FBlueprintCompilerCppBackendBase::EmitFileBeginning(const FString& CleanNam
 	{
 		EmitIncludeHeader(EmitterContext.Body, *PCHFilename, false);
 	}
-	/*
-	const FString MainHeaderFilename = FEmitHelper::GetGameMainHeaderFilename();
-	if (!MainHeaderFilename.IsEmpty() && MainHeaderFilename != PCHFilename)
+	else
 	{
-		EmitIncludeHeader(EmitterContext.Body, *MainHeaderFilename, false);
+		// Used when generated code is not in a separate module
+		const FString MainHeaderFilename = FEmitHelper::GetGameMainHeaderFilename();
+		if (!MainHeaderFilename.IsEmpty())
+		{
+			EmitIncludeHeader(EmitterContext.Body, *MainHeaderFilename, false);
+		}
 	}
-	*/
 	EmitIncludeHeader(EmitterContext.Body, *CleanName, true);
 	EmitIncludeHeader(bIncludeCodeHelpersInHeader ? EmitterContext.Header : EmitterContext.Body, TEXT("GeneratedCodeHelpers"), true);
 	EmitIncludeHeader(EmitterContext.Header, TEXT("Blueprint/BlueprintSupport"), true);
 
 	FBackendHelperUMG::AdditionalHeaderIncludeForWidget(EmitterContext);
-	EmitterContext.Header.AddLine(FBackendHelperAnim::AddHeaders(EmitterContext.GetCurrentlyGeneratedClass()));
+	FBackendHelperAnim::AddHeaders(EmitterContext);
 
 	TSet<FString> AlreadyIncluded;
 	AlreadyIncluded.Add(CleanName);
diff --git a/Engine/Source/Developer/BlueprintCompilerCppBackend/Private/BlueprintCompilerCppBackendGatherDependencies.cpp b/Engine/Source/Developer/BlueprintCompilerCppBackend/Private/BlueprintCompilerCppBackendGatherDependencies.cpp
--- a/Engine/Source/Developer/BlueprintCompilerCppBackend/Private/BlueprintCompilerCppBackendGatherDependencies.cpp
+++ b/Engine/Source/Developer/BlueprintCompilerCppBackend/Private/BlueprintCompilerCppBackendGatherDependencies.cpp
@@ -142,8 +142,21 @@ struct FFindHeadersToInclude : public FGatherConvertedClassDependenciesHelperBas
 		}
 
 		auto ObjAsField = Cast<UField>(InObject);
+		if (!ObjAsField && InObject->IsA<UBlueprintFunctionLibrary>())
+		{
+			ObjAsField = InObject->GetClass();
+		}
+
 		if (ObjAsField)
 		{
+			if (ObjAsField->IsA<UProperty>())
+			{
+				ObjAsField = ObjAsField->GetOwnerStruct();
+			}
+			if (ObjAsField->IsA<UFunction>())
+			{
+				ObjAsField = ObjAsField->GetOwnerClass();
+			}
 			IncludeTheHeaderInBody(ObjAsField);
 		}
 
diff --git a/Engine/Source/Developer/BlueprintCompilerCppBackend/Private/BlueprintCompilerCppBackendUtils.cpp b/Engine/Source/Developer/BlueprintCompilerCppBackend/Private/BlueprintCompilerCppBackendUtils.cpp
--- a/Engine/Source/Developer/BlueprintCompilerCppBackend/Private/BlueprintCompilerCppBackendUtils.cpp
+++ b/Engine/Source/Developer/BlueprintCompilerCppBackend/Private/BlueprintCompilerCppBackendUtils.cpp
@@ -5,6 +5,19 @@
 #include "EdGraphSchema_K2.h"
 #include "IBlueprintCompilerCppBackendModule.h" // for OnPCHFilenameQuery()
 
+FString GetPathPostfix(const UObject* ForObject)
+{
+	FString FullAssetName = ForObject->GetOutermost()->GetPathName();
+	if (FullAssetName.StartsWith(TEXT("/Temp/__TEMP_BP__"), ESearchCase::CaseSensitive))
+	{
+		FullAssetName.RemoveFromStart(TEXT("/Temp/__TEMP_BP__"), ESearchCase::CaseSensitive);
+	}
+	FString AssetName = FPackageName::GetLongPackageAssetName(FullAssetName);
+	// append a hash of the path, this uniquely identifies assets with the same name, but different folders:
+	FullAssetName.RemoveFromEnd(AssetName);
+	return FString::Printf(TEXT("%u"), FCrc::MemCrc32(*FullAssetName, FullAssetName.Len()*sizeof(TCHAR)));
+}
+
 FString FEmitterLocalContext::GenerateUniqueLocalName()
 {
 	const FString UniqueNameBase = TEXT("__Local__");
@@ -238,6 +251,31 @@ FString FEmitterLocalContext::ExportCppDeclaration(const UProperty* Property, EE
 			ActualCppTypePtr = &ActualCppType;
 		}
 	}
+	else if (const UByteProperty* ByteProperty = Cast<const UByteProperty>(Property))
+	{
+		if (ByteProperty->Enum != nullptr)
+		{
+			const bool bIsEnumClass  = (ByteProperty->Enum->GetCppForm() == UEnum::ECppForm::EnumClass);
+			const bool bIsNativeEnum = (ByteProperty->Enum->GetClass()   == UEnum::StaticClass()); // cannot use RF_Native flag, because in UHT the flag is not set
+
+			// unfortunately, we cannot tell if the native enum is typed or not 
+			// (uint8 or not); since UHT allows untyped enums to be exposed to
+			// Blueprints, but errors when they're used for a native properties, 
+			// then we need to wrap all native class enums in TEnumAsByte<> 
+			// (to take preventative measures against non uint8 enums)
+			if (bIsNativeEnum && bIsEnumClass)
+			{
+				const FString EnumAsBytePrefix = TEXT("TEnumAsByte<");
+
+				ActualCppType = ByteProperty->GetCPPType(/*ExtendedTypeText =*/nullptr, InExportCPPFlags | CPPF_BlueprintCppBackend);
+				if (!ActualCppType.StartsWith(EnumAsBytePrefix))
+				{
+					ActualCppType = FString::Printf(TEXT("%s %s >"), *EnumAsBytePrefix, *ActualCppType);
+				}
+				ActualCppTypePtr = &ActualCppType;
+			}
+		}
+	}
 
 	if (ArrayProperty)
 	{
@@ -262,7 +300,8 @@ FString FEmitterLocalContext::ExportCppDeclaration(const UProperty* Property, EE
 	}
 
 	FStringOutputDevice Out;
-	Property->ExportCppDeclaration(Out, DeclarationType, nullptr, InExportCPPFlags, bSkipParameterName, ActualCppTypePtr, ActualExtendedTypePtr);
+	const FString ActualNativeName = bSkipParameterName ? FString() : FEmitHelper::GetCppName(Property);
+	Property->ExportCppDeclaration(Out, DeclarationType, nullptr, InExportCPPFlags, bSkipParameterName, ActualCppTypePtr, ActualExtendedTypePtr, &ActualNativeName);
 	return FString(Out);
 
 }
@@ -288,7 +327,7 @@ FString FEmitHelper::GetCppName(const UField* Field, bool bUInterface)
 		}
 		else
 		{
-			return ::UnicodeToCPPIdentifier(*AsStruct->GetName(), false, AsStruct->GetPrefixCPP());
+			return ::UnicodeToCPPIdentifier(*AsStruct->GetName(), false, AsStruct->GetPrefixCPP()) + GetPathPostfix(AsStruct);
 		}
 	}
 	else if (auto AsProperty = Cast<UProperty>(Field))
@@ -298,8 +337,27 @@ FString FEmitHelper::GetCppName(const UField* Field, bool bUInterface)
 			if ((Cast<UBlueprintGeneratedClass>(Owner) ||
 				!Owner->IsNative()))
 			{
-				const bool bIsParameter = AsProperty->HasAnyPropertyFlags(CPF_Parm);
-				return ::UnicodeToCPPIdentifier(AsProperty->GetName(), AsProperty->HasAnyPropertyFlags(CPF_Deprecated), bIsParameter ? TEXT("bpp__") : TEXT("bpv__"));
+				FString VarPrefix;
+				
+				const bool bIsUberGraphVariable = Owner->IsA<UBlueprintGeneratedClass>() && AsProperty->HasAllPropertyFlags(CPF_Transient | CPF_DuplicateTransient);
+				if (bIsUberGraphVariable)
+				{
+					int32 InheritenceLevel = 0;
+
+					const UStruct* StructIt = Owner->GetSuperStruct();
+					while ((StructIt != nullptr) && !StructIt->IsNative())
+					{
+						++InheritenceLevel;
+						StructIt = StructIt->GetSuperStruct();
+					}
+					VarPrefix = FString::Printf(TEXT("b%dl__"), InheritenceLevel);
+				}
+				else
+				{
+					const bool bIsParameter = AsProperty->HasAnyPropertyFlags(CPF_Parm);
+					VarPrefix = bIsParameter ? TEXT("bpp__") : TEXT("bpv__");
+				}
+				return ::UnicodeToCPPIdentifier(AsProperty->GetName(), AsProperty->HasAnyPropertyFlags(CPF_Deprecated), *VarPrefix);
 			}
 		}
 		return AsProperty->GetNameCPP();
@@ -335,15 +393,40 @@ FString FEmitHelper::HandleRepNotifyFunc(const UProperty* Property)
 	check(Property);
 	if (HasAllFlags(Property->PropertyFlags, CPF_Net | CPF_RepNotify))
 	{
-		return FString::Printf(TEXT("ReplicatedUsing=%s"), *Property->RepNotifyFunc.ToString());
+		if (Property->RepNotifyFunc != NAME_None)
+		{
+			return FString::Printf(TEXT("ReplicatedUsing=\"%s\""), *Property->RepNotifyFunc.ToString());
+		}
+		else
+		{
+			UE_LOG(LogK2Compiler, Warning, TEXT("Invalid RepNotifyFunc in %s"), *GetPathNameSafe(Property));
+		}
 	}
-	else if (HasAllFlags(Property->PropertyFlags, CPF_Net))
+
+	if (HasAllFlags(Property->PropertyFlags, CPF_Net))
 	{
 		return TEXT("Replicated");
 	}
 	return FString();
 }
 
+bool FEmitHelper::IsMetaDataValid(const FName Name, const FString& Value)
+{
+	static const FName UIMin(TEXT("UIMin"));
+	static const FName UIMax(TEXT("UIMax"));
+	static const FName ClampMin(TEXT("ClampMin"));
+	static const FName ClampMax(TEXT("ClampMax"));
+	if ((Name == UIMin)
+		|| (Name == UIMax)
+		|| (Name == ClampMin)
+		|| (Name == ClampMax))
+	{
+		// those MD require no warning
+		return Value.IsNumeric();
+	}
+	return true;
+}
+
 bool FEmitHelper::MetaDataCanBeNative(const FName MetaDataName, const UField* Field)
 {
 	if (MetaDataName == TEXT("ModuleRelativePath"))
@@ -373,7 +456,7 @@ FString FEmitHelper::HandleMetaData(const UField* Field, bool AddCategory, const
 	{
 		for (auto& Pair : *ValuesMap)
 		{
-			if (!MetaDataCanBeNative(Pair.Key, Field))
+			if (!MetaDataCanBeNative(Pair.Key, Field) || !IsMetaDataValid(Pair.Key, Pair.Value))
 			{
 				continue;
 			}
@@ -600,7 +683,7 @@ FString FEmitHelper::GetBaseFilename(const UObject* AssetObj)
 			Char = TCHAR('x');
 		}
 	}
-
+	Postfix += GetPathPostfix(AssetObj);
 	return AssetName + Postfix;
 }
 
@@ -957,99 +1040,105 @@ FString FEmitHelper::LiteralTerm(FEmitterLocalContext& EmitterContext, const FEd
 	return CustomValue;
 }
 
-FString FEmitHelper::DefaultValue(FEmitterLocalContext& EmitterContext, const FEdGraphPinType& Type)
+FString FEmitHelper::PinTypeToNativeType(const FEdGraphPinType& Type)
 {
-	if (Type.bIsArray)
+	// A temp uproperty should be generated?
+	auto PinTypeToNativeTypeInner = [](const FEdGraphPinType& InType) -> FString
 	{
-		auto PinTypeToNativeType = [](const FEdGraphPinType& InType) -> FString
+		auto Schema = GetDefault<UEdGraphSchema_K2>();
+		if (UEdGraphSchema_K2::PC_String == InType.PinCategory)
 		{
-			// A temp property should be generated?
-
-			auto Schema = GetDefault<UEdGraphSchema_K2>();
-			if (UEdGraphSchema_K2::PC_String == InType.PinCategory)
-			{
-				return TEXT("FString");
-			}
-			else if (UEdGraphSchema_K2::PC_Boolean == InType.PinCategory)
-			{
-				return TEXT("bool");
-			}
-			else if((UEdGraphSchema_K2::PC_Byte == InType.PinCategory) || (UEdGraphSchema_K2::PC_Enum == InType.PinCategory))
-			{
-				if (UEnum* Enum = Cast<UEnum>(InType.PinSubCategoryObject.Get()))
-				{
-					return FEmitHelper::GetCppName(Enum);
-				}
-				return TEXT("uint8");
-			}
-			else if(UEdGraphSchema_K2::PC_Int == InType.PinCategory)
-			{
-				return TEXT("int32");
-			}
-			else if(UEdGraphSchema_K2::PC_Float == InType.PinCategory)
-			{
-				return TEXT("float");
-			}
-			else if(UEdGraphSchema_K2::PC_Float == InType.PinCategory)
-			{
-				return TEXT("float");
-			}
-			else if(UEdGraphSchema_K2::PC_Name == InType.PinCategory)
-			{
-				return TEXT("FName");
-			}
-			else if(UEdGraphSchema_K2::PC_Text == InType.PinCategory)
+			return TEXT("FString");
+		}
+		else if (UEdGraphSchema_K2::PC_Boolean == InType.PinCategory)
+		{
+			return TEXT("bool");
+		}
+		else if ((UEdGraphSchema_K2::PC_Byte == InType.PinCategory) || (UEdGraphSchema_K2::PC_Enum == InType.PinCategory))
+		{
+			if (UEnum* Enum = Cast<UEnum>(InType.PinSubCategoryObject.Get()))
 			{
-				return TEXT("FText");
+				FString FullyQualifiedEnumName = (!Enum->CppType.IsEmpty()) ? Enum->CppType : FEmitHelper::GetCppName(Enum);
+				return FString::Printf(TEXT("TEnumAsByte<%s>"), *FullyQualifiedEnumName);
 			}
-			else if(UEdGraphSchema_K2::PC_Struct == InType.PinCategory)
+			return TEXT("uint8");
+		}
+		else if (UEdGraphSchema_K2::PC_Int == InType.PinCategory)
+		{
+			return TEXT("int32");
+		}
+		else if (UEdGraphSchema_K2::PC_Float == InType.PinCategory)
+		{
+			return TEXT("float");
+		}
+		else if (UEdGraphSchema_K2::PC_Float == InType.PinCategory)
+		{
+			return TEXT("float");
+		}
+		else if (UEdGraphSchema_K2::PC_Name == InType.PinCategory)
+		{
+			return TEXT("FName");
+		}
+		else if (UEdGraphSchema_K2::PC_Text == InType.PinCategory)
+		{
+			return TEXT("FText");
+		}
+		else if (UEdGraphSchema_K2::PC_Struct == InType.PinCategory)
+		{
+			if (UScriptStruct* Struct = Cast<UScriptStruct>(InType.PinSubCategoryObject.Get()))
 			{
-				if (UScriptStruct* Struct = Cast<UScriptStruct>(InType.PinSubCategoryObject.Get()))
-				{
-					return FEmitHelper::GetCppName(Struct);
-				}
+				return FEmitHelper::GetCppName(Struct);
 			}
-			else if(UEdGraphSchema_K2::PC_Class == InType.PinCategory)
+		}
+		else if (UEdGraphSchema_K2::PC_Class == InType.PinCategory)
+		{
+			if (UClass* Class = Cast<UClass>(InType.PinSubCategoryObject.Get()))
 			{
-				if (UClass* Class = Cast<UClass>(InType.PinSubCategoryObject.Get()))
-				{
-					return FString::Printf(TEXT("TSubclassOf<%s>"), *FEmitHelper::GetCppName(Class));
-				}
+				return FString::Printf(TEXT("TSubclassOf<%s>"), *FEmitHelper::GetCppName(Class));
 			}
-			else if (UEdGraphSchema_K2::PC_AssetClass == InType.PinCategory)
+		}
+		else if (UEdGraphSchema_K2::PC_AssetClass == InType.PinCategory)
+		{
+			if (UClass* Class = Cast<UClass>(InType.PinSubCategoryObject.Get()))
 			{
-				if (UClass* Class = Cast<UClass>(InType.PinSubCategoryObject.Get()))
-				{
-					return FString::Printf(TEXT("TAssetSubclassOf<%s>"), *FEmitHelper::GetCppName(Class));
-				}
+				return FString::Printf(TEXT("TAssetSubclassOf<%s>"), *FEmitHelper::GetCppName(Class));
 			}
-			else if (UEdGraphSchema_K2::PC_Interface == InType.PinCategory)
+		}
+		else if (UEdGraphSchema_K2::PC_Interface == InType.PinCategory)
+		{
+			if (UClass* Class = Cast<UClass>(InType.PinSubCategoryObject.Get()))
 			{
-				if (UClass* Class = Cast<UClass>(InType.PinSubCategoryObject.Get()))
-				{
-					return FString::Printf(TEXT("TScriptInterface<%s>"), *FEmitHelper::GetCppName(Class));
-				}
+				return FString::Printf(TEXT("TScriptInterface<%s>"), *FEmitHelper::GetCppName(Class));
 			}
-			else if (UEdGraphSchema_K2::PC_Asset == InType.PinCategory)
+		}
+		else if (UEdGraphSchema_K2::PC_Asset == InType.PinCategory)
+		{
+			if (UClass* Class = Cast<UClass>(InType.PinSubCategoryObject.Get()))
 			{
-				if (UClass* Class = Cast<UClass>(InType.PinSubCategoryObject.Get()))
-				{
-					return FString::Printf(TEXT("TAssetPtr<%s>"), *FEmitHelper::GetCppName(Class));
-				}
+				return FString::Printf(TEXT("TAssetPtr<%s>"), *FEmitHelper::GetCppName(Class));
 			}
-			else if (UEdGraphSchema_K2::PC_Object == InType.PinCategory)
+		}
+		else if (UEdGraphSchema_K2::PC_Object == InType.PinCategory)
+		{
+			if (UClass* Class = Cast<UClass>(InType.PinSubCategoryObject.Get()))
 			{
-				if (UClass* Class = Cast<UClass>(InType.PinSubCategoryObject.Get()))
-				{
-					return FString::Printf(TEXT("%s*"), *FEmitHelper::GetCppName(Class));
-				}
+				return FString::Printf(TEXT("%s*"), *FEmitHelper::GetCppName(Class));
 			}
-			UE_LOG(LogK2Compiler, Error, TEXT("FEmitHelper::DefaultValue cannot generate an array type"));
-			return FString{};
-		};
+		}
+		UE_LOG(LogK2Compiler, Error, TEXT("FEmitHelper::DefaultValue cannot generate an array type"));
+		return FString{};
+	};
 
+	FString InnerTypeName = PinTypeToNativeTypeInner(Type);
+	return Type.bIsArray ? FString::Printf(TEXT("TArray<%s>"), *InnerTypeName) : InnerTypeName;
+}
+
+FString FEmitHelper::DefaultValue(FEmitterLocalContext& EmitterContext, const FEdGraphPinType& Type)
+{
+	if (Type.bIsArray)
+	{
 		const FString InnerTypeName = PinTypeToNativeType(Type);
-		return FString::Printf(TEXT("TArray<%s>{}"), *InnerTypeName);
+		return FString::Printf(TEXT("%s{}"), *InnerTypeName);
 	}
 
 	return LiteralTerm(EmitterContext, Type, FString(), nullptr);
@@ -1201,7 +1290,7 @@ FString FEmitHelper::GenerateGetPropertyByName(FEmitterLocalContext& EmitterCont
 }
 
 FString FEmitHelper::AccessInaccessibleProperty(FEmitterLocalContext& EmitterContext, const UProperty* Property
-	, const FString& ContextStr, const FString& ContextAdressOp, const FString& StaticArrayIdx)
+	, const FString& ContextStr, const FString& ContextAdressOp, int32 StaticArrayIdx)
 {
 	check(Property);
 	const FString PropertyLocalName = GenerateGetPropertyByName(EmitterContext, Property);
@@ -1210,10 +1299,119 @@ FString FEmitHelper::AccessInaccessibleProperty(FEmitterLocalContext& EmitterCon
 		| EPropertyExportCPPFlags::CPPF_BlueprintCppBackend;
 	const FString TypeDeclaration = EmitterContext.ExportCppDeclaration(Property, EExportedDeclaration::Parameter, CppTemplateTypeFlags, true);
 
-	return FString::Printf(TEXT("(*(%s->ContainerPtrToValuePtr<%s>(%s(%s)%s)))")
+	return FString::Printf(TEXT("(*(%s->ContainerPtrToValuePtr<%s>(%s(%s), %d)))")
 		, *PropertyLocalName
 		, *TypeDeclaration
 		, *ContextAdressOp
 		, *ContextStr
-		, *StaticArrayIdx);
+		, StaticArrayIdx);
+}
+
+namespace HelperWithoutEditorOnlyMembers
+{
+	// This code works properly as long, as all fields in structures are UProperties
+
+	static int32 SizeOfStructWithoutParent(const UStruct* InStruct);
+	static int32 SizeOfWholeStruct(const UStruct* InStruct);
+	static int32 OffsetWithoutEditorOnlyMembers(const UProperty* Property);
+	static int32 SizeWithoutEditorOnlyMembers(const UProperty* InProperty);
+
+	int32 SizeOfStructWithoutParent(const UStruct* InStruct)
+	{
+		check(InStruct);
+		int32 SizeResult = 0;
+		for (const UField* Field = InStruct->Children; (Field != nullptr) && (Field->GetOuter() == InStruct); Field = Field->Next)
+		{
+			const UProperty* Property = Cast<const UProperty>(Field);
+			if (Property && !Property->IsEditorOnlyProperty())
+			{
+				SizeResult += SizeWithoutEditorOnlyMembers(Property);
+			}
+		}
+		return SizeResult;
+	}
+
+	int32 SizeOfWholeStruct(const UStruct* InStruct)
+	{
+		check(InStruct);
+		int32 SizeResult = 0;
+		for (const UStruct* Struct = InStruct; Struct; Struct = Struct->GetSuperStruct())
+		{
+			SizeResult += SizeOfStructWithoutParent(Struct);
+		}
+		return SizeResult;
+	}
+
+	int32 OffsetWithoutEditorOnlyMembers(const UProperty* InProperty)
+	{
+		check(InProperty);
+		const UProperty* Property = InProperty->GetOwnerProperty();
+		ensure(Property == InProperty); // it's hard to tell what user expects otherwise 
+
+		UStruct* OwnerStruct = Property->GetOwnerStruct();
+		int32 SizeResult = 0;
+
+		// size of all super
+		for (const UStruct* StructIt = OwnerStruct->GetSuperStruct(); StructIt; StructIt = StructIt->GetSuperStruct())
+		{
+			SizeResult += SizeOfStructWithoutParent(StructIt);
+		}
+
+		// size of properties before this
+		const UField* Field = OwnerStruct->Children;
+		for (; Field && (Field->GetOuter() == OwnerStruct) && (Field != Property); Field = Field->Next)
+		{
+			const UProperty* LocProperty = Cast<const UProperty>(Field);
+			if (LocProperty && !LocProperty->IsEditorOnlyProperty())
+			{
+				SizeResult += SizeWithoutEditorOnlyMembers(LocProperty);
+			}
+		}
+		ensure(Field == Property);
+		return SizeResult;
+	}
+
+	int32 SizeWithoutEditorOnlyMembers(const UProperty* InProperty)
+	{
+		check(InProperty);
+		const UProperty* Property = InProperty->GetOwnerProperty();
+		ensure(Property == InProperty); // it's hard to tell what user expects otherwise 
+		if (const UStructProperty* StructyProperty = Cast<const UStructProperty>(Property))
+		{
+			return SizeOfWholeStruct(StructyProperty->Struct) * Property->ArrayDim;
+		}
+
+		return Property->GetSize();
+	}
+
+}
+
+FString FEmitHelper::AccessInaccessiblePropertyUsingOffset(FEmitterLocalContext& EmitterContext, const UProperty* Property
+	, const FString& ContextStr, const FString& ContextAdressOp, int32 StaticArrayIdx)
+{
+	check(Property);
+
+	const int32 PropertyOffsetWithoutEditorOnlyMembers = HelperWithoutEditorOnlyMembers::OffsetWithoutEditorOnlyMembers(Property);
+	const int32 PropertySizeWithoutEditorOnlyMembers = HelperWithoutEditorOnlyMembers::SizeWithoutEditorOnlyMembers(Property);
+	const int32 ElementSizeWithoutEditorOnlyMembers = PropertySizeWithoutEditorOnlyMembers / Property->ArrayDim;
+	ensure(PropertySizeWithoutEditorOnlyMembers == (ElementSizeWithoutEditorOnlyMembers * Property->ArrayDim));
+
+	EmitterContext.AddLine(FString::Printf(TEXT("// Offset of property: %s"), *Property->GetPathName()));
+	/*
+	EmitterContext.AddLine(FString::Printf(TEXT("static_assert(STRUCT_OFFSET(%s, %s) == 0x%08X, \"Wrong generated offset\");")
+		, *FEmitHelper::GetCppName(Property->GetOwnerStruct())
+		, *FEmitHelper::GetCppName(Property)
+		, PropertyOffsetWithoutEditorOnlyMembers));
+	*/
+	const uint32 CppTemplateTypeFlags = EPropertyExportCPPFlags::CPPF_CustomTypeName
+		| EPropertyExportCPPFlags::CPPF_NoConst | EPropertyExportCPPFlags::CPPF_NoRef | EPropertyExportCPPFlags::CPPF_NoStaticArray
+		| EPropertyExportCPPFlags::CPPF_BlueprintCppBackend;
+	const FString TypeDeclaration = EmitterContext.ExportCppDeclaration(Property, EExportedDeclaration::Parameter, CppTemplateTypeFlags, true);
+	return FString::Printf(TEXT("(*(AccessPrivateProperty<%s>(%s(%s), 0x%08X, 0x%08X, %d)))")
+		, *TypeDeclaration
+		, *ContextAdressOp
+		, *ContextStr
+		, PropertyOffsetWithoutEditorOnlyMembers
+		, ElementSizeWithoutEditorOnlyMembers
+		, StaticArrayIdx);
 }
diff --git a/Engine/Source/Developer/BlueprintCompilerCppBackend/Private/BlueprintCompilerCppBackendUtils.h b/Engine/Source/Developer/BlueprintCompilerCppBackend/Private/BlueprintCompilerCppBackendUtils.h
--- a/Engine/Source/Developer/BlueprintCompilerCppBackend/Private/BlueprintCompilerCppBackendUtils.h
+++ b/Engine/Source/Developer/BlueprintCompilerCppBackend/Private/BlueprintCompilerCppBackendUtils.h
@@ -189,6 +189,8 @@ struct FEmitHelper
 
 	static bool HasAllFlags(uint64 Flags, uint64 FlagsToCheck);
 
+	static bool IsMetaDataValid(const FName Name, const FString& Value);
+
 	static FString HandleRepNotifyFunc(const UProperty* Property);
 
 	static bool MetaDataCanBeNative(const FName MetaDataName, const UField* Field);
@@ -219,6 +221,8 @@ struct FEmitHelper
 
 	static FString DefaultValue(FEmitterLocalContext& EmitterContext, const FEdGraphPinType& Type);
 
+	static FString PinTypeToNativeType(const FEdGraphPinType& InType);
+
 	static UFunction* GetOriginalFunction(UFunction* Function);
 
 	static bool ShouldHandleAsNativeEvent(UFunction* Function);
@@ -240,7 +244,11 @@ struct FEmitHelper
 	static FString GenerateGetPropertyByName(FEmitterLocalContext& EmitterContext, const UProperty* Property);
 
 	static FString AccessInaccessibleProperty(FEmitterLocalContext& EmitterContext, const UProperty* Property
-		, const FString& ContextStr, const FString& ContextAdressOp, const FString& StaticArrayIdx = FString(TEXT(", 0")));
+		, const FString& ContextStr, const FString& ContextAdressOp, int32 StaticArrayIdx = 0);
+
+	// This code works properly as long, as all fields in structures are UProperties!
+	static FString AccessInaccessiblePropertyUsingOffset(FEmitterLocalContext& EmitterContext, const UProperty* Property
+		, const FString& ContextStr, const FString& ContextAdressOp, int32 StaticArrayIdx = 0);
 
 };
 
@@ -295,7 +303,7 @@ struct FBackendHelperUMG
 
 struct FBackendHelperAnim
 {
-	static FString AddHeaders(UClass* GeneratedClass);
+	static void AddHeaders(FEmitterLocalContext& EmitterContext);
 
 	static void CreateAnimClassData(FEmitterLocalContext& Context);
 };
\ No newline at end of file
diff --git a/Engine/Source/Developer/BlueprintCompilerCppBackend/Private/BlueprintCompilerCppBackendValueHelper.cpp b/Engine/Source/Developer/BlueprintCompilerCppBackend/Private/BlueprintCompilerCppBackendValueHelper.cpp
--- a/Engine/Source/Developer/BlueprintCompilerCppBackend/Private/BlueprintCompilerCppBackendValueHelper.cpp
+++ b/Engine/Source/Developer/BlueprintCompilerCppBackend/Private/BlueprintCompilerCppBackendValueHelper.cpp
@@ -45,7 +45,13 @@ void FEmitDefaultValueHelper::OuterGenerate(FEmitterLocalContext& Context
 			|| (!Property->Identical_InContainer(DataContainer, OptionalDefaultDataContainer, ArrayIndex) && !IsInstancedSubobjectLambda(ArrayIndex)))
 		{
 			FString PathToMember;
-			auto PropertyOwnerAsBPGC = Cast<UBlueprintGeneratedClass>(Property->GetOwnerClass());
+			UBlueprintGeneratedClass* PropertyOwnerAsBPGC = Cast<UBlueprintGeneratedClass>(Property->GetOwnerClass());
+			UScriptStruct* PropertyOwnerAsScriptStruct = Cast<UScriptStruct>(Property->GetOwnerStruct());
+			const bool bNoexportProperty = PropertyOwnerAsScriptStruct
+				&& PropertyOwnerAsScriptStruct->IsNative()
+				&& (PropertyOwnerAsScriptStruct->StructFlags & STRUCT_NoExport)
+				// && !PropertyOwnerAsScriptStruct->GetBoolMetaData(TEXT("BlueprintType"))
+				&& ensure(EPropertyAccessOperator::Dot == AccessOperator);
 			if (PropertyOwnerAsBPGC && !Context.Dependencies.WillClassBeConverted(PropertyOwnerAsBPGC))
 			{
 				ensure(EPropertyAccessOperator::None != AccessOperator);
@@ -55,13 +61,14 @@ void FEmitDefaultValueHelper::OuterGenerate(FEmitterLocalContext& Context
 				PathToMember = FString::Printf(TEXT("FUnconvertedWrapper__%s(%s).GetRef__%s()"), *FEmitHelper::GetCppName(PropertyOwnerAsBPGC), *ContainerStr
 					, *UnicodeToCPPIdentifier(Property->GetName(), false, nullptr));
 			}
-			else if (Property->HasAnyPropertyFlags(CPF_NativeAccessSpecifierPrivate) || (!bAllowProtected && Property->HasAnyPropertyFlags(CPF_NativeAccessSpecifierProtected)))
+			else if (bNoexportProperty || Property->HasAnyPropertyFlags(CPF_NativeAccessSpecifierPrivate) || (!bAllowProtected && Property->HasAnyPropertyFlags(CPF_NativeAccessSpecifierProtected)))
 			{
 				ensure(EPropertyAccessOperator::None != AccessOperator);
 				const FString OperatorStr = (EPropertyAccessOperator::Dot == AccessOperator) ? TEXT("&") : TEXT("");
 				const FString ContainerStr = (EPropertyAccessOperator::None == AccessOperator) ? TEXT("this") : OuterPath;
-				const FString StaticArrayIndexArg = FString::Printf(TEXT(", %d"), ArrayIndex);
-				const FString GetPtrStr = FEmitHelper::AccessInaccessibleProperty(Context, Property, ContainerStr, OperatorStr, StaticArrayIndexArg);
+				const FString GetPtrStr = bNoexportProperty
+					? FEmitHelper::AccessInaccessiblePropertyUsingOffset(Context, Property, ContainerStr, OperatorStr, ArrayIndex)
+					: FEmitHelper::AccessInaccessibleProperty(Context, Property, ContainerStr, OperatorStr, ArrayIndex);
 				PathToMember = Context.GenerateUniqueLocalName();
 				Context.AddLine(FString::Printf(TEXT("auto& %s = %s;"), *PathToMember, *GetPtrStr));
 			}
@@ -264,6 +271,8 @@ struct FNonativeComponentData
 	////
 	FString ParentVariableName;
 	bool bSetNativeCreationMethod;
+	/** Socket/Bone that Component might attach to */
+	FName AttachToName;
 
 	FNonativeComponentData()
 		: ComponentTemplate(nullptr)
@@ -282,7 +291,9 @@ struct FNonativeComponentData
 
 		if (!ParentVariableName.IsEmpty())
 		{
-			Context.AddLine(FString::Printf(TEXT("%s->AttachParent = %s;"), *NativeVariablePropertyName, *ParentVariableName));
+			const FString SocketName = (AttachToName == NAME_None) ? FString() : FString::Printf(TEXT(", TEXT(\"%s\")"), *AttachToName.ToString());
+			Context.AddLine(FString::Printf(TEXT("%s->AttachTo(%s %s);"), *NativeVariablePropertyName, *ParentVariableName, *SocketName));
+			// AttachTo is called first in case some properties will be overridden.
 		}
 
 		UClass* ComponentClass = ComponentTemplate->GetClass();
@@ -293,6 +304,7 @@ struct FNonativeComponentData
 				, reinterpret_cast<const uint8*>(ObjectToCompare)
 				, FEmitDefaultValueHelper::EPropertyAccessOperator::Pointer);
 		}
+
 	}
 };
 
@@ -358,6 +370,7 @@ FString FEmitDefaultValueHelper::HandleNonNativeComponent(FEmitterLocalContext&
 					ParentVariableName = Context.FindGloballyMappedObject(ParentComponentTemplate, USceneComponent::StaticClass());
 				}
 				NonativeComponentData.ParentVariableName = ParentVariableName;
+				NonativeComponentData.AttachToName = Node->AttachToName;
 			}
 			NonativeComponentData.ObjectToCompare = ObjectToCompare;
 			ComponenntsToInit.Add(NonativeComponentData);
diff --git a/Engine/Source/Developer/BlueprintNativeCodeGen/Private/BlueprintNativeCodeGenManifest.cpp b/Engine/Source/Developer/BlueprintNativeCodeGen/Private/BlueprintNativeCodeGenManifest.cpp
--- a/Engine/Source/Developer/BlueprintNativeCodeGen/Private/BlueprintNativeCodeGenManifest.cpp
+++ b/Engine/Source/Developer/BlueprintNativeCodeGen/Private/BlueprintNativeCodeGenManifest.cpp
@@ -8,7 +8,6 @@
 #include "IBlueprintCompilerCppBackendModule.h" // for ConstructBaseFilename()
 #include "Engine/Blueprint.h"					// for GeneratedClass
 #include "JsonObjectConverter.h"
-#include "NativeCodeGenCommandlineParams.h"
 #include "Serialization/JsonReader.h"
 #include "Serialization/JsonSerializer.h"
 #include "Serialization/JsonWriter.h"
@@ -28,7 +27,6 @@ namespace BlueprintNativeCodeGenManifestImpl
 	static const FString HeaderSubDir         = TEXT("Public");
 	static const FString CppSubDir            = TEXT("Private");
 	static const FString ModuleBuildFileExt   = TEXT(".Build.cs");
-	static const FString FallbackPluginName   = TEXT("NativizedScript");
 	static const FString PreviewFilePostfix   = TEXT("-Preview");
 	static const FString PluginFileExt        = TEXT(".uplugin");
 	static const FString SourceSubDir         = TEXT("Source");
@@ -220,7 +218,7 @@ static bool BlueprintNativeCodeGenManifestImpl::GatherModuleDependencies(const U
 	const FLinkerLoad* PkgLinker = FLinkerLoad::FindExistingLinkerForPackage(AssetPackage);
 	const bool bFoundLinker = (PkgLinker != nullptr);
 
-	if (ensure(bFoundLinker))
+	if (ensureMsgf(bFoundLinker, TEXT("Failed to identify the asset package that '%s' belongs to."), *AssetObj->GetName()))
 	{
 		for (const FObjectImport& PkgImport : PkgLinker->ImportMap)
 		{
@@ -297,13 +295,6 @@ FConvertedAssetRecord::FConvertedAssetRecord(const FAssetData& AssetInfo, const
 	GeneratedHeaderPath = BlueprintNativeCodeGenManifestImpl::GenerateSourceFileSavePath(TargetPaths, AssetInfo, FBlueprintNativeCodeGenPaths::HFile);
 }
 
-//------------------------------------------------------------------------------
-bool FConvertedAssetRecord::IsValid()
-{
-	// every conversion will have a header file (interfaces don't have implementation files)
-	return ToAssetRef().IsValid() && (AssetType != nullptr) && !GeneratedHeaderPath.IsEmpty();
-}
-
 //------------------------------------------------------------------------------
 FStringAssetReference FConvertedAssetRecord::ToAssetRef() const
 {
@@ -315,17 +306,11 @@ FStringAssetReference FConvertedAssetRecord::ToAssetRef() const
  ******************************************************************************/
  
 //------------------------------------------------------------------------------
-FUnconvertedDependencyRecord::FUnconvertedDependencyRecord(const FAssetData& AssetInfo, const FBlueprintNativeCodeGenPaths& TargetPaths)
-	: GeneratedWrapperPath(BlueprintNativeCodeGenManifestImpl::GenerateUnconvertedWrapperPath(TargetPaths, AssetInfo))
+FUnconvertedDependencyRecord::FUnconvertedDependencyRecord(const FString& InGeneratedWrapperPath)
+	: GeneratedWrapperPath(InGeneratedWrapperPath)
 {
 }
 
-//------------------------------------------------------------------------------
-bool FUnconvertedDependencyRecord::IsValid()
-{
-	return !GeneratedWrapperPath.IsEmpty();
-}
-
 /*******************************************************************************
  * FBlueprintNativeCodeGenPaths
  ******************************************************************************/
@@ -410,105 +395,35 @@ FString FBlueprintNativeCodeGenPaths::RuntimePCHFilename() const
 
 //------------------------------------------------------------------------------
 FBlueprintNativeCodeGenManifest::FBlueprintNativeCodeGenManifest()
-	: OutputDir(FPaths::ConvertRelativePathToFull( FPaths::Combine(*FPaths::Combine(*FPaths::GameIntermediateDir(), TEXT("Plugins")), *BlueprintNativeCodeGenManifestImpl::FallbackPluginName) ))
 {
-	PluginName = FPaths::GetBaseFilename(OutputDir);
-	ManifestFilePath = FPaths::Combine(*GetTargetDir(), *FBlueprintNativeCodeGenPaths::GetDefaultManifestPath());
 }
 
 //------------------------------------------------------------------------------
-FBlueprintNativeCodeGenManifest::FBlueprintNativeCodeGenManifest(const FNativeCodeGenCommandlineParams& CommandlineParams)
+FBlueprintNativeCodeGenManifest::FBlueprintNativeCodeGenManifest(const FString& InPluginName, const FString& InOutputDir)
 {
 	using namespace BlueprintNativeCodeGenManifestImpl;
-	bool const bLoadExistingManifest = !CommandlineParams.bWipeRequested || CommandlineParams.OutputDir.IsEmpty();
 	
-	PluginName = CommandlineParams.PluginName;
-	OutputDir  = CommandlineParams.OutputDir;
+	PluginName = InPluginName;
+	OutputDir = InOutputDir;
 	if (FPaths::IsRelative(OutputDir))
 	{
 		FPaths::MakePathRelativeTo(OutputDir, *FPaths::GameDir());
 	}
 
-	if (!CommandlineParams.ManifestFilePath.IsEmpty())
-	{
-		ManifestFilePath = CommandlineParams.ManifestFilePath;
-	}
-	else
+	if (!ensure(!OutputDir.IsEmpty()))
 	{
-		if (!ensure(!PluginName.IsEmpty()))
-		{
-			PluginName = FallbackPluginName;
-		}
-		if (!ensure(!OutputDir.IsEmpty()))
-		{
-			OutputDir = FPaths::Combine(*FPaths::GameIntermediateDir(), *PluginName);
-		}
-
-		FString ManifestFilename = FBlueprintNativeCodeGenPaths::GetDefaultManifestPath();
-		if (CommandlineParams.bPreviewRequested)
-		{
-			ManifestFilename.InsertAt(ManifestFilename.Find(ManifestFileExt, ESearchCase::CaseSensitive, ESearchDir::FromEnd), PreviewFilePostfix);
-		}
-		ManifestFilePath = FPaths::Combine(*GetTargetDir(), *ManifestFilename);
+		OutputDir = FPaths::Combine(*FPaths::GameIntermediateDir(), *PluginName);
 	}
-	
-	// incorporate an existing manifest (in case we're only re-converting a 
-	// handful of assets and adding them into an existing module)
-	if (bLoadExistingManifest && LoadManifest(ManifestFilePath, this))
-	{
-		// if they specified a separate plugin path, lets make sure we use that
-		// over what was found in the existing manifest
-		if (!CommandlineParams.OutputDir.IsEmpty())
-		{
-			if (CommandlineParams.bWipeRequested)
-			{
-				OutputDir = CommandlineParams.OutputDir;
-			}
-			else
-			{
-				const FString ExpectedPath = GetComparibleDirPath(CommandlineParams.OutputDir);
-
-				FString TargetPath = GetTargetDir();
-				TargetPath = GetComparibleDirPath(TargetPath);
-
-				if ( !FPaths::IsSamePath(ExpectedPath, TargetPath) )
-				{
-					UE_LOG(LogNativeCodeGenManifest, Error, 
-						TEXT("The existing manifest's plugin path does not match what was specified via the commandline."));
-				}
-			}
-		}
 
-		const bool bNewPluginNameRequested = !CommandlineParams.PluginName.IsEmpty() && (CommandlineParams.PluginName != PluginName);
-		if (bNewPluginNameRequested && !CommandlineParams.bPreviewRequested)
-		{
-			// delete the old plugin file (if one exists)
-			IFileManager::Get().Delete(*GetTargetPaths().PluginFilePath());
-		}
-
-		// if we were only interested in obtaining the plugin path
-		if (CommandlineParams.bWipeRequested)
-		{
-			PluginName = CommandlineParams.PluginName;
-			Clear();
-		}
-		else
-		{
-			if (bNewPluginNameRequested)
-			{
-				UE_LOG(LogNativeCodeGenManifest, Warning, TEXT("The specified plugin name (%s) doesn't match the existing one (%s). Overridding with the new name."),
-					*CommandlineParams.PluginName, *PluginName);
-				PluginName = CommandlineParams.PluginName;
-			}
-		}
-	}
+	FString ManifestFilename = FBlueprintNativeCodeGenPaths::GetDefaultManifestPath();
+	ManifestFilePath = FPaths::Combine(*GetTargetDir(), *ManifestFilename);
 }
 
 //------------------------------------------------------------------------------
 FBlueprintNativeCodeGenManifest::FBlueprintNativeCodeGenManifest(const FString& ManifestFilePathIn)
 	: ManifestFilePath(ManifestFilePathIn)
 {
-	BlueprintNativeCodeGenManifestImpl::LoadManifest(ManifestFilePathIn, this);
+	ensureAlwaysMsgf(BlueprintNativeCodeGenManifestImpl::LoadManifest(ManifestFilePathIn, this), TEXT("Missing Manifest for Blueprint code generation: %s"), *ManifestFilePath );
 }
 
 //------------------------------------------------------------------------------
@@ -520,69 +435,38 @@ FBlueprintNativeCodeGenPaths FBlueprintNativeCodeGenManifest::GetTargetPaths() c
 //------------------------------------------------------------------------------
 FConvertedAssetRecord& FBlueprintNativeCodeGenManifest::CreateConversionRecord(const FAssetId Key, const FAssetData& AssetInfo)
 {
-	const FBlueprintNativeCodeGenPaths TargetPaths = GetTargetPaths();
+	// It is an error to consider a record converted and unconverted. This is probably not negotiable. In order to leave the door
+	// open for parallelism we need to be able to trivially and consistently determine whether an asset will be converted without
+	// discovering things as we go.
+	check(UnconvertedDependencies.Find(Key) == nullptr);
+	// Similarly we should not convert things over and over and over:
+	if (FConvertedAssetRecord* Existing = ConvertedAssets.Find(Key))
+	{
+		return *Existing;
+	}
 
+	const FBlueprintNativeCodeGenPaths TargetPaths = GetTargetPaths();
 	UClass* AssetType = AssetInfo.GetClass();
 	// load the asset (if it isn't already)
 	const UObject* AssetObj = AssetInfo.GetAsset();
-
-	FConvertedAssetRecord* ConversionRecordPtr = ConvertedAssets.Find(Key);
-	if (ConversionRecordPtr == nullptr)
-	{
-		UnconvertedDependencies.Remove(Key);
-		ConversionRecordPtr = &ConvertedAssets.Add(Key, FConvertedAssetRecord(AssetInfo, TargetPaths));
-	}
-	else if (!ensure(AssetType == ConversionRecordPtr->AssetType))
-	{
-		UE_LOG(LogNativeCodeGenManifest, Warning, TEXT("Existing manifest entry for '%s' has a different type, overriding with: '%s'"),
-			*AssetInfo.ObjectPath.ToString(), *AssetType->GetName());
-
-		ConversionRecordPtr->AssetType = AssetType;
-		ConversionRecordPtr->TargetObjPath = BlueprintNativeCodeGenManifestImpl::GetTargetObjectPath(AssetInfo);
-	}
-	
-	FConvertedAssetRecord& ConversionRecord = *ConversionRecordPtr;
-	if (!ConversionRecord.IsValid())
-	{
-		if (ConversionRecord.GeneratedHeaderPath.IsEmpty())
-		{
-			ConversionRecord.GeneratedCppPath    = BlueprintNativeCodeGenManifestImpl::GenerateSourceFileSavePath(TargetPaths, AssetInfo, FBlueprintNativeCodeGenPaths::CppFile);
-			ConversionRecord.GeneratedHeaderPath = BlueprintNativeCodeGenManifestImpl::GenerateSourceFileSavePath(TargetPaths, AssetInfo, FBlueprintNativeCodeGenPaths::HFile);
-		}
-		ensure(ConversionRecord.IsValid());
-	}
-
-	return ConversionRecord;
+	FConvertedAssetRecord* ConversionRecord = &ConvertedAssets.Add(Key, FConvertedAssetRecord(AssetInfo, TargetPaths));
+	return *ConversionRecord;
 }
 
 //------------------------------------------------------------------------------
-FUnconvertedDependencyRecord& FBlueprintNativeCodeGenManifest::CreateUnconvertedDependencyRecord(const FAssetId UnconvertedAssetKey, const FAssetData& AssetInfo, const FAssetId ReferencingAsset)
+FUnconvertedDependencyRecord& FBlueprintNativeCodeGenManifest::CreateUnconvertedDependencyRecord(const FAssetId UnconvertedAssetKey, const FAssetData& AssetInfo)
 {
-	FUnconvertedDependencyRecord* RecordPtr = UnconvertedDependencies.Find(UnconvertedAssetKey);
-	if (RecordPtr == nullptr)
-	{
-		const FBlueprintNativeCodeGenPaths TargetPaths = GetTargetPaths();
-		RecordPtr = &UnconvertedDependencies.Add(UnconvertedAssetKey, FUnconvertedDependencyRecord(AssetInfo, TargetPaths));
-	}
-	check(RecordPtr != nullptr);
+	// It is an error to create a record multiple times because it is not obvious what the client wants to happen.
+	check(ConvertedAssets.Find(UnconvertedAssetKey) == nullptr);
 
-	FUnconvertedDependencyRecord& DependencyRecord = *RecordPtr;
-	if (!DependencyRecord.IsValid())
+	if (FUnconvertedDependencyRecord* Existing = UnconvertedDependencies.Find(UnconvertedAssetKey))
 	{
-		const FBlueprintNativeCodeGenPaths TargetPaths = GetTargetPaths();
-		DependencyRecord.GeneratedWrapperPath = BlueprintNativeCodeGenManifestImpl::GenerateUnconvertedWrapperPath(TargetPaths, AssetInfo);
-
-		ensure(DependencyRecord.IsValid());
+		return *Existing;
 	}
 
-	DependencyRecord.ReferencedBy.AddUnique(ReferencingAsset);
-	return DependencyRecord;
-}
-
-//------------------------------------------------------------------------------
-const FConvertedAssetRecord* FBlueprintNativeCodeGenManifest::FindConversionRecord(const FAssetId AssetId) const
-{
-	return ConvertedAssets.Find(AssetId);
+	const FBlueprintNativeCodeGenPaths TargetPaths = GetTargetPaths();
+	FUnconvertedDependencyRecord* RecordPtr = &UnconvertedDependencies.Add(UnconvertedAssetKey, FUnconvertedDependencyRecord(BlueprintNativeCodeGenManifestImpl::GenerateUnconvertedWrapperPath(TargetPaths, AssetInfo)));
+	return *RecordPtr;
 }
 
 //------------------------------------------------------------------------------
@@ -592,14 +476,19 @@ void FBlueprintNativeCodeGenManifest::GatherModuleDependencies(UPackage* Package
 }
 
 //------------------------------------------------------------------------------
-bool FBlueprintNativeCodeGenManifest::Save() const
-{
-	return SaveAs(*ManifestFilePath);
-}
-
-//------------------------------------------------------------------------------
-bool FBlueprintNativeCodeGenManifest::SaveAs(const TCHAR* Filename) const
+bool FBlueprintNativeCodeGenManifest::Save(int32 Id) const
 {
+	FString FullFilename;
+	if (Id != -1)
+	{
+		FullFilename = *(ManifestFilePath + FString::FromInt(Id));
+	}
+	else
+	{
+		FullFilename = *ManifestFilePath;
+	}
+	
+	const TCHAR* Filename = *FullFilename;
 	TSharedRef<FJsonObject> JsonObject = MakeShareable(new FJsonObject());
 
 	if (FJsonObjectConverter::UStructToJsonObject(FBlueprintNativeCodeGenManifest::StaticStruct(), this, JsonObject,
@@ -617,9 +506,8 @@ bool FBlueprintNativeCodeGenManifest::SaveAs(const TCHAR* Filename) const
 	return false;
 }
 
-void FBlueprintNativeCodeGenManifest::Merge(const TCHAR* Filename)
+void FBlueprintNativeCodeGenManifest::Merge(const FBlueprintNativeCodeGenManifest& OtherManifest)
 {
-	FBlueprintNativeCodeGenManifest OtherManifest = FBlueprintNativeCodeGenManifest(FString(Filename));
 	for (auto& Entry : OtherManifest.ModuleDependencies)
 	{
 		ModuleDependencies.AddUnique(Entry);
@@ -630,9 +518,9 @@ void FBlueprintNativeCodeGenManifest::Merge(const TCHAR* Filename)
 		ConvertedAssets.Add(Entry.Key, Entry.Value);
 	}
 
-	for (auto& Entry : UnconvertedDependencies)
+	for (auto& Entry : OtherManifest.UnconvertedDependencies)
 	{
-		OtherManifest.UnconvertedDependencies.Add(Entry.Key, Entry.Value);
+		UnconvertedDependencies.Add(Entry.Key, Entry.Value);
 	}
 }
 
diff --git a/Engine/Source/Developer/BlueprintNativeCodeGen/Private/BlueprintNativeCodeGenManifest.h b/Engine/Source/Developer/BlueprintNativeCodeGen/Private/BlueprintNativeCodeGenManifest.h
--- a/Engine/Source/Developer/BlueprintNativeCodeGen/Private/BlueprintNativeCodeGenManifest.h
+++ b/Engine/Source/Developer/BlueprintNativeCodeGen/Private/BlueprintNativeCodeGenManifest.h
@@ -22,11 +22,6 @@ public:
 	FConvertedAssetRecord() {}
 	FConvertedAssetRecord(const FAssetData& AssetInfo, const FBlueprintNativeCodeGenPaths& TargetPaths);
 
-	/**
-	 * @return True if this record demonstates a successful conversion; false if it is incomplete 
-	 */
-	bool IsValid();
-
 	/**
 	 * @return 
 	 */
@@ -59,17 +54,11 @@ struct FUnconvertedDependencyRecord
 
 public:
 	FUnconvertedDependencyRecord() {}
-	FUnconvertedDependencyRecord(const FAssetData& AssetInfo, const FBlueprintNativeCodeGenPaths& TargetPaths);
-
-	/**  */
-	bool IsValid();
+	FUnconvertedDependencyRecord(const FString& InGeneratedWrapperPath);
 
 public:
 	UPROPERTY()
 	FString GeneratedWrapperPath;
-
-	UPROPERTY()
-	TArray<FName> ReferencedBy;
 };
 
 /*******************************************************************************
@@ -128,7 +117,7 @@ struct FBlueprintNativeCodeGenManifest
 
 public: 
 	FBlueprintNativeCodeGenManifest(); // default ctor for USTRUCT() system
-	FBlueprintNativeCodeGenManifest(const FNativeCodeGenCommandlineParams& CommandlineParams);
+	FBlueprintNativeCodeGenManifest(const FString& InPluginName, const FString& InOutputDir);
 	FBlueprintNativeCodeGenManifest(const FString& ManifestFilePath);
 
 	/**
@@ -154,16 +143,7 @@ public:
 	 * @param  ReferencingAsset    
 	 * @return 
 	 */
-	FUnconvertedDependencyRecord& CreateUnconvertedDependencyRecord(const FAssetId UnconvertedAssetKey, const FAssetData& AssetInfo, const FAssetId ReferencingAsset);
-
-	/**
-	 * Searches the manifest's conversion records for an entry pertaining to the
-	 * specified asset.
-	 * 
-	 * @param  AssetPath    Serves as a lookup key for the asset in question.
-	 * @return A pointer to the found conversion record, null if one doesn't exist.
-	 */
-	const FConvertedAssetRecord* FindConversionRecord(const FAssetId AssetId) const;
+	FUnconvertedDependencyRecord& CreateUnconvertedDependencyRecord(const FAssetId UnconvertedAssetKey, const FAssetData& AssetInfo);
 
 	/**
 	 * Records module dependencies for the specified asset.
@@ -182,25 +162,24 @@ public:
 	 */
 	const FConversionRecord& GetConversionRecord() const { return ConvertedAssets; }
 
+	/**
+	* @return A list of all asset conversion that require stubs.
+	*/
+	const FUnconvertedRecord& GetUnconvertedDependencies() const { return UnconvertedDependencies; }
+
 	/**
 	 * Saves this manifest as json, to its target destination (which it was 
 	 * setup with).
 	 * 
+	 * @param Id an identifier used to uniquely identify this manifest, used by multiprocess cook
 	 * @return True if the save was a success, otherwise false.
 	 */
-	bool Save() const;
-
-	/**
-	* Saves this manifest as json, to the provided destination 
-	*
-	* @return True if the save was a success, otherwise false.
-	*/
-	bool SaveAs(const TCHAR* Filename) const;
+	bool Save(int32 Id) const;
 
 	/**
-	* Merges the manifest saved in Filename into the current manifest
+	* Merges the OtherManifest into the current manifest
 	*/
-	void Merge(const TCHAR* Filename);
+	void Merge(const FBlueprintNativeCodeGenManifest& OtherManifest);
 
 private:
 	/**
diff --git a/Engine/Source/Developer/BlueprintNativeCodeGen/Private/BlueprintNativeCodeGenModule.cpp b/Engine/Source/Developer/BlueprintNativeCodeGen/Private/BlueprintNativeCodeGenModule.cpp
--- a/Engine/Source/Developer/BlueprintNativeCodeGen/Private/BlueprintNativeCodeGenModule.cpp
+++ b/Engine/Source/Developer/BlueprintNativeCodeGen/Private/BlueprintNativeCodeGenModule.cpp
@@ -7,7 +7,7 @@
 #include "BlueprintNativeCodeGenModule.h"
 #include "BlueprintNativeCodeGenUtils.h"
 #include "IBlueprintCompilerCppBackendModule.h"
-#include "NativeCodeGenCommandlineParams.h"
+#include "Engine/Blueprint.h" // for EBlueprintType
 
 /*******************************************************************************
 * NativizationCookControllerImpl
@@ -15,146 +15,127 @@
 
 namespace NativizationCookControllerImpl
 {
-#if WITH_EDITORONLY_DATA
-	static const FString ConvertAssetsCommand = TEXT("NativizeAssets");
-	static const FString ConvertedManifestParam = TEXT("NativizedAssetManifest=");
-	static const FString GeneratedPluginParam = TEXT("NativizedAssetPlugin=");
+	// If you change this plugin name you must update logic in CookCommand.Automation.cs
+	static const TCHAR* DefaultPluginName = TEXT("NativizedAssets");
+}
 
-	//--------------------------------------------------------------------------
-	static bool IsRunningCookCommandlet()
+/*******************************************************************************
+ * FBlueprintNativeCodeGenModule
+ ******************************************************************************/
+ 
+class FBlueprintNativeCodeGenModule : public IBlueprintNativeCodeGenModule
+									, public IBlueprintNativeCodeGenCore
+{
+public:
+	FBlueprintNativeCodeGenModule()
 	{
-		// @TODO: figure out how to determine if this is the cook commandlet
-		return IsRunningCommandlet();
 	}
-#endif // WITH_EDITORONLY_DATA
-}
 
-/*******************************************************************************
-* UBlueprintNativeCodeGenConfig
-*******************************************************************************/
+	//~ Begin IBlueprintNativeCodeGenModule interface
+	virtual void Convert(UPackage* Package, ESavePackageResult ReplacementType, const TCHAR* PlatformName) override;
+	virtual void SaveManifest(int32 Id = -1) override;
+	virtual void MergeManifest(int32 ManifestIdentifier) override;
+	virtual void FinalizeManifest() override;
+protected:
+	virtual void Initialize(const FNativeCodeGenInitData& InitData) override;
+	virtual void InitializeForRerunDebugOnly(const TArray< TPair< FString, FString > >& CodegenTargets) override;
+	//~ End IBlueprintNativeCodeGenModule interface
 
-//------------------------------------------------------------------------------
-UBlueprintNativeCodeGenConfig::UBlueprintNativeCodeGenConfig(FObjectInitializer const& ObjectInitializer)
-: Super(ObjectInitializer)
-{
-}
+	//~ Begin FScriptCookReplacmentCoordinator interface
+	virtual EReplacementResult IsTargetedForReplacement(const UPackage* Package) const override;
+	virtual EReplacementResult IsTargetedForReplacement(const UObject* Object) const override;
+	virtual UClass* FindReplacedClass(const UClass* Class) const override;
+	//~ End FScriptCookReplacmentCoordinator interface
+private:
+	FBlueprintNativeCodeGenManifest& GetManifest(const TCHAR* PlatformName);
 
-/*******************************************************************************
-* FCookCommandParams
-******************************************************************************/
+	TMap< FString, TUniquePtr<FBlueprintNativeCodeGenManifest> > Manifests;
 
-//------------------------------------------------------------------------------
-FCookCommandParams::FCookCommandParams(const FString& Commandline)
-: bRunConversion(false)
+	TArray<FString> ExcludedAssetTypes;
+	TArray<FString> ExcludedBlueprintTypes;
+	TArray<FString> TargetPlatformNames;
+};
+
+void FBlueprintNativeCodeGenModule::Initialize(const FNativeCodeGenInitData& InitData)
 {
-	using namespace NativizationCookControllerImpl;
+	GConfig->GetArray(TEXT("BlueprintNativizationSettings"), TEXT("ExcludedAssetTypes"), ExcludedAssetTypes, GEditorIni);
+	GConfig->GetArray(TEXT("BlueprintNativizationSettings"), TEXT("ExcludedBlueprintTypes"), ExcludedBlueprintTypes, GEditorIni);
 
-#if WITH_EDITORONLY_DATA
-	if (IsRunningCookCommandlet())
-	{
-		IFileManager& FileManager = IFileManager::Get();
-		bRunConversion = FParse::Param(*Commandline, *ConvertAssetsCommand);
+	IBlueprintNativeCodeGenCore::Register(this);
 
-		if (FParse::Value(*Commandline, *ConvertedManifestParam, ManifestFilePath))
-		{
-			bRunConversion |= !FileManager.FileExists(*ManifestFilePath);
-		}
+	// Each platform will need a manifest, because each platform could cook different assets:
+	for (auto& Platform : InitData.CodegenTargets)
+	{
+		const TCHAR* TargetDirectory = *Platform.Value;
+		FString OutputPath = FPaths::Combine(TargetDirectory, NativizationCookControllerImpl::DefaultPluginName);
 
-		FString PluginPath;
-		if (FParse::Value(*Commandline, *GeneratedPluginParam, PluginPath))
-		{
-			if (PluginPath.EndsWith(TEXT(".uplugin")))
-			{
-				PluginName = FPaths::GetBaseFilename(PluginPath);
-				PluginPath = FPaths::GetPath(PluginPath);
-			}
-			OutputPath = PluginPath;
-			bRunConversion = true;
-		}
+		Manifests.Add(FString(*Platform.Key), TUniquePtr<FBlueprintNativeCodeGenManifest>(new FBlueprintNativeCodeGenManifest(NativizationCookControllerImpl::DefaultPluginName, OutputPath)));
 
-		if (bRunConversion && ManifestFilePath.IsEmpty())
-		{
-			ManifestFilePath = FPaths::Combine(*FPaths::Combine(*FPaths::GameIntermediateDir(), TEXT("Cook")), TEXT("BlueprintConversionManifest"));
-		}
+		TargetPlatformNames.Add(Platform.Key);
 	}
-#endif // WITH_EDITORONLY_DATA
-}
 
-//------------------------------------------------------------------------------
-TArray<FString> FCookCommandParams::ToConversionParams() const
-{
-	TArray<FString> ConversionCmdParams;
-	if (!PluginName.IsEmpty())
+	IBlueprintCompilerCppBackendModule& BackEndModule = (IBlueprintCompilerCppBackendModule&)IBlueprintCompilerCppBackendModule::Get();
+	auto& ConversionQueryDelegate = BackEndModule.OnIsTargetedForConversionQuery();
+	auto ShouldConvert = [](const UObject* AssetObj)
 	{
-		ConversionCmdParams.Add(FString::Printf(TEXT("pluginName=%s"), *PluginName));
-	}
-	if (!OutputPath.IsEmpty())
-	{
-		ConversionCmdParams.Add(FString::Printf(TEXT("output=\"%s\""), *OutputPath));
-	}
-	if (!ManifestFilePath.IsEmpty())
-	{
-		ConversionCmdParams.Add(FString::Printf(TEXT("manifest=\"%s\""), *ManifestFilePath));
-	}
-	return ConversionCmdParams;
+		if (IBlueprintNativeCodeGenCore::Get())
+		{
+			EReplacementResult ReplacmentResult = IBlueprintNativeCodeGenCore::Get()->IsTargetedForReplacement(AssetObj);
+			return ReplacmentResult == EReplacementResult::ReplaceCompletely;
+		}
+		return false;
+	};
+	ConversionQueryDelegate.BindStatic(ShouldConvert);
 }
 
-/*******************************************************************************
- * FBlueprintNativeCodeGenModule
- ******************************************************************************/
- 
-class FBlueprintNativeCodeGenModule : public IBlueprintNativeCodeGenModule
+void FBlueprintNativeCodeGenModule::InitializeForRerunDebugOnly(const TArray< TPair< FString, FString > >& CodegenTargets)
 {
-public:
-	FBlueprintNativeCodeGenModule()
-		: Manifest(nullptr)
+	for (const auto& Platform : CodegenTargets)
 	{
-		FCookCommandParams CookParams(FCommandLine::Get());
-		FNativeCodeGenCommandlineParams CommandLineParams(CookParams.ToConversionParams());
-		Manifest = new FBlueprintNativeCodeGenManifest(CommandLineParams);
+		// load the old manifest:
+		const TCHAR* TargetDirectory = *Platform.Value;
+		FString OutputPath = FPaths::Combine(TargetDirectory, NativizationCookControllerImpl::DefaultPluginName, *FBlueprintNativeCodeGenPaths::GetDefaultManifestPath());
+		Manifests.Add(FString(*Platform.Key), TUniquePtr<FBlueprintNativeCodeGenManifest>(new FBlueprintNativeCodeGenManifest(FPaths::ConvertRelativePathToFull(OutputPath))));
+		//FBlueprintNativeCodeGenManifest OldManifest(FPaths::ConvertRelativePathToFull(OutputPath));
+		// reconvert every assets listed in the manifest:
+		for (const auto& ConversionTarget : GetManifest(*Platform.Key).GetConversionRecord())
+		{
+			// load the package:
+			UPackage* Package = LoadPackage(nullptr, *ConversionTarget.Value.TargetObjPath, LOAD_None);
 
-		IBlueprintCompilerCppBackendModule& BackEndModule = (IBlueprintCompilerCppBackendModule&)IBlueprintCompilerCppBackendModule::Get();
-		auto& BackendPCHQuery = BackEndModule.OnPCHFilenameQuery();
+			// reconvert it:
+			Convert(Package, ESavePackageResult::ReplaceCompletely, *Platform.Key);
+		}
 
-		FBlueprintNativeCodeGenPaths TargetPaths = GetManifest().GetTargetPaths();
-		BackendPCHQuery.BindLambda([TargetPaths]()->FString
+		// reconvert every unconverted dependency listed in the manifest:
+		for (const auto& ConversionTarget : GetManifest(*Platform.Key).GetUnconvertedDependencies())
 		{
-			return TargetPaths.RuntimePCHFilename();
-		});
+			// load the package:
+			UPackage* Package = LoadPackage(nullptr, *ConversionTarget.Key.GetPlainNameString(), LOAD_None);
 
-		auto& ConversionQueryDelegate = BackEndModule.OnIsTargetedForConversionQuery();
-		auto ShouldConvert = [](const UObject* AssetObj)
-		{
-			if (FScriptCookReplacementCoordinator::Get())
-			{
-				EReplacementResult ReplacmentResult = FScriptCookReplacementCoordinator::Get()->IsTargetedForReplacement(AssetObj);
-				return ReplacmentResult == EReplacementResult::ReplaceCompletely;
-			}
-			return false;
-		};
-		ConversionQueryDelegate.BindStatic(ShouldConvert);
+			// reconvert it:
+			Convert(Package, ESavePackageResult::GenerateStub, *Platform.Key);
+		}
 	}
 
-	//~ Begin IBlueprintNativeCodeGenModule interface
-	virtual void Convert(UPackage* Package, EReplacementResult ReplacementType) override;
-	virtual void SaveManifest(const TCHAR* Filename) override;
-	virtual void MergeManifest(const TCHAR* Filename) override;
-	virtual void FinalizeManifest() override;
-	//~ End IBlueprintNativeCodeGenModule interface
-private:
-	FBlueprintNativeCodeGenManifest& GetManifest();
-
-	FBlueprintNativeCodeGenManifest* Manifest;
-};
+}
 
-FBlueprintNativeCodeGenManifest& FBlueprintNativeCodeGenModule::GetManifest()
+FBlueprintNativeCodeGenManifest& FBlueprintNativeCodeGenModule::GetManifest(const TCHAR* PlatformName)
 {
-	check(Manifest);
-	return *Manifest;
+	FString PlatformNameStr(PlatformName);
+	TUniquePtr<FBlueprintNativeCodeGenManifest>* Result = Manifests.Find(PlatformNameStr);
+	check(Result->IsValid());
+	return **Result;
 }
 
-void FBlueprintNativeCodeGenModule::Convert(UPackage* Package, EReplacementResult ReplacementType)
+void FBlueprintNativeCodeGenModule::Convert(UPackage* Package, ESavePackageResult CookResult, const TCHAR* PlatformName)
 {
+	if (CookResult != ESavePackageResult::ReplaceCompletely && CookResult != ESavePackageResult::GenerateStub)
+	{
+		// nothing to convert
+		return;
+	}
+
 	// Find the struct/enum to convert:
 	UStruct* Struct = nullptr;
 	UEnum* Enum = nullptr;
@@ -186,32 +167,38 @@ void FBlueprintNativeCodeGenModule::Convert(UPackage* Package, EReplacementResul
 		return;
 	}
 
+	IBlueprintCompilerCppBackendModule& BackEndModule = (IBlueprintCompilerCppBackendModule&)IBlueprintCompilerCppBackendModule::Get();
+	auto& BackendPCHQuery = BackEndModule.OnPCHFilenameQuery();
+
+	FBlueprintNativeCodeGenPaths TargetPaths = GetManifest(PlatformName).GetTargetPaths();
+	BackendPCHQuery.BindLambda([TargetPaths]()->FString
+	{
+		return TargetPaths.RuntimePCHFilename();
+	});
+
 	const IAssetRegistry& Registry = FModuleManager::LoadModuleChecked<FAssetRegistryModule>("AssetRegistry").Get();
-	auto GenerateWrapperSrcFile = [](const FUnconvertedDependencyRecord& UnconvertedDependency, UClass* Class)
+
+	if (CookResult == ESavePackageResult::GenerateStub)
 	{
+		FAssetData AssetInfo = Registry.GetAssetByObjectPath(*Struct->GetPathName());
 		FString FileContents;
 		TUniquePtr<IBlueprintCompilerCppBackend> Backend_CPP(IBlueprintCompilerCppBackendModuleInterface::Get().Create());
-		FileContents = Backend_CPP->GenerateWrapperForClass(Class);
+		// Apparently we can only generate wrappers for classes, so any logic that results in non classes requesting
+		// wrappers will fail here:
 
+		FileContents = Backend_CPP->GenerateWrapperForClass(CastChecked<UClass>(Struct));
 		if (!FileContents.IsEmpty())
 		{
-			FFileHelper::SaveStringToFile(FileContents, *UnconvertedDependency.GeneratedWrapperPath);
+			FFileHelper::SaveStringToFile(FileContents, *(GetManifest(PlatformName).CreateUnconvertedDependencyRecord(AssetInfo.PackageName, AssetInfo).GeneratedWrapperPath));
 		}
-	};
-
-	if (ReplacementType == EReplacementResult::GenerateStub)
-	{
-		check(ReplacementType == EReplacementResult::GenerateStub);
-		FAssetData AssetInfo = Registry.GetAssetByObjectPath(*Struct->GetPathName());
-		GenerateWrapperSrcFile(GetManifest().CreateUnconvertedDependencyRecord(AssetInfo.PackageName, AssetInfo, FName()), CastChecked<UClass>(Struct));
 		// The stub we generate still may have dependencies on other modules, so make sure the module dependencies are 
 		// still recorded so that the .build.cs is generated correctly. Without this you'll get include related errors 
 		// (or possibly linker errors) in stub headers:
-		GetManifest().GatherModuleDependencies(Package);
+		GetManifest(PlatformName).GatherModuleDependencies(Package);
 	}
 	else
 	{
-		check(ReplacementType == EReplacementResult::ReplaceCompletely);
+		check(CookResult == ESavePackageResult::ReplaceCompletely);
 		// convert:
 		UField* ForConversion = Enum;
 		if (ForConversion == nullptr)
@@ -220,7 +207,7 @@ void FBlueprintNativeCodeGenModule::Convert(UPackage* Package, EReplacementResul
 		}
 
 		FAssetData AssetInfo = Registry.GetAssetByObjectPath(*ForConversion->GetPathName());
-		FConvertedAssetRecord& ConversionRecord = GetManifest().CreateConversionRecord(ForConversion->GetFName(), AssetInfo);
+		FConvertedAssetRecord& ConversionRecord = GetManifest(PlatformName).CreateConversionRecord(*ForConversion->GetPathName(), AssetInfo);
 		TSharedPtr<FString> HeaderSource(new FString());
 		TSharedPtr<FString> CppSource(new FString());
 
@@ -259,29 +246,177 @@ void FBlueprintNativeCodeGenModule::Convert(UPackage* Package, EReplacementResul
 		check(bSuccess);
 		if (bSuccess)
 		{
-			GetManifest().GatherModuleDependencies(Package);
+			GetManifest(PlatformName).GatherModuleDependencies(Package);
 		}
 	}
+
+	BackendPCHQuery.Unbind();
 }
 
-void FBlueprintNativeCodeGenModule::SaveManifest(const TCHAR* Filename)
+void FBlueprintNativeCodeGenModule::SaveManifest(int32 Id )
 {
-	GetManifest().SaveAs(Filename);
+	for (auto& PlatformName : TargetPlatformNames)
+	{
+		GetManifest(*PlatformName).Save(Id);
+	}
 }
 
-void FBlueprintNativeCodeGenModule::MergeManifest(const TCHAR* Filename)
+void FBlueprintNativeCodeGenModule::MergeManifest(int32 ManifestIdentifier)
 {
-	GetManifest().Merge(Filename);
+	for (auto& PlatformName : TargetPlatformNames)
+	{
+		FBlueprintNativeCodeGenManifest& CurrentManifest = GetManifest(*PlatformName);
+		FBlueprintNativeCodeGenManifest OtherManifest = FBlueprintNativeCodeGenManifest(CurrentManifest.GetTargetPaths().ManifestFilePath() + FString::FromInt(ManifestIdentifier));
+		CurrentManifest.Merge(OtherManifest);
+	}
 } 
 
 void FBlueprintNativeCodeGenModule::FinalizeManifest()
 {
-	GetManifest().Save();
+	for(auto& PlatformName : TargetPlatformNames)
+	{
+		GetManifest(*PlatformName).Save(-1);
+		check(FBlueprintNativeCodeGenUtils::FinalizePlugin(GetManifest(*PlatformName)));
+	}
+}
+
+UClass* FBlueprintNativeCodeGenModule::FindReplacedClass(const UClass* Class) const
+{
+	// we're only looking to replace class types:
+	while (Class)
+	{
+		if (Class == UUserDefinedEnum::StaticClass())
+		{
+			return UEnum::StaticClass();
+		}
+		if (Class == UUserDefinedStruct::StaticClass())
+		{
+			return UScriptStruct::StaticClass();
+		}
+		if (Class == UBlueprintGeneratedClass::StaticClass())
+		{
+			return UDynamicClass::StaticClass();
+		}
+
+		Class = Class->GetSuperClass();
+	}
+
+	return nullptr;
+}
+
+EReplacementResult FBlueprintNativeCodeGenModule::IsTargetedForReplacement(const UPackage* Package) const
+{
+	// non-native packages with enums and structs should be converted, unless they are blacklisted:
+	UStruct* Struct = nullptr;
+	UEnum* Enum = nullptr;
+	TArray<UObject*> Objects;
+	GetObjectsWithOuter(Package, Objects, false);
+	for (auto Entry : Objects)
+	{
+		Struct = Cast<UStruct>(Entry);
+		Enum = Cast<UEnum>(Entry);
+		if (Struct || Enum)
+		{
+			break;
+		}
+	}
+
+	UObject* Target = Struct;
+	if (Target == nullptr)
+	{
+		Target = Enum;
+	}
+	return IsTargetedForReplacement(Target);
+}
+
+EReplacementResult FBlueprintNativeCodeGenModule::IsTargetedForReplacement(const UObject* Object) const
+{
+	const UStruct* Struct = Cast<UStruct>(Object);
+	const UEnum* Enum = Cast<UEnum>(Object);
+
+	if (Struct == nullptr && Enum == nullptr)
+	{
+		return EReplacementResult::DontReplace;
+	}
+
+	if (const UClass* BlueprintClass = Cast<UClass>(Struct))
+	{
+		if (UBlueprint* Blueprint = Cast<UBlueprint>(BlueprintClass->ClassGeneratedBy))
+		{
+			const EBlueprintType UnconvertableBlueprintTypes[] = {
+				//BPTYPE_Const,		// WTF is a "const" Blueprint?
+				BPTYPE_MacroLibrary,
+				BPTYPE_LevelScript,
+			};
+
+			EBlueprintType BlueprintType = Blueprint->BlueprintType;
+			for (int32 TypeIndex = 0; TypeIndex < ARRAY_COUNT(UnconvertableBlueprintTypes); ++TypeIndex)
+			{
+				if (BlueprintType == UnconvertableBlueprintTypes[TypeIndex])
+				{
+					return EReplacementResult::DontReplace;
+				}
+			}
+		}
+	}
+
+	auto IsObjectFromDeveloperPackage = [](const UObject* Obj) -> bool
+	{
+		return Obj && Obj->GetOutermost()->HasAllPackagesFlags(PKG_Developer);
+	};
+
+	auto IsDeveloperObject = [&](const UObject* Obj) -> bool
+	{
+		if (Obj)
+		{
+			if (IsObjectFromDeveloperPackage(Obj))
+			{
+				return true;
+			}
+			const UStruct* StructToTest = Obj->IsA<UStruct>() ? CastChecked<const UStruct>(Obj) : Obj->GetClass();
+			for (; StructToTest; StructToTest = StructToTest->GetSuperStruct())
+			{
+				if (IsObjectFromDeveloperPackage(StructToTest))
+				{
+					return true;
+				}
+			}
+		}
+		return false;
+	};
+
+	if (Object && (IsEditorOnlyObject(Object) || IsDeveloperObject(Object)))
+	{
+		UE_LOG(LogBlueprintCodeGen, Warning, TEXT("Object %s depends on Editor or Development stuff. Is shouldn't be cooked."), *GetPathNameSafe(Object));
+		return EReplacementResult::DontReplace;
+	}
+
+	// check blacklists:
+	// we can't use FindObject, because we may be converting a type while saving
+	if ((Struct && ExcludedAssetTypes.Find(Struct->GetPathName()) != INDEX_NONE) ||
+		(Enum && ExcludedAssetTypes.Find(Enum->GetPathName()) != INDEX_NONE))
+	{
+		return EReplacementResult::GenerateStub;
+	}
+
+	EReplacementResult Result = EReplacementResult::ReplaceCompletely;
+	while (Struct)
+	{
+		// This happens because the cooker incorrectly cooks editor only packages. Specifically happens for the blackjack sample
+		// project due to a FStringAssetReference in BaseEditor.ini:
+		if (Struct->RootPackageHasAnyFlags(PKG_EditorOnly))
+		{
+			return EReplacementResult::DontReplace;
+		}
+
+		if (ExcludedBlueprintTypes.Find(Struct->GetPathName()) != INDEX_NONE)
+		{
+			Result = EReplacementResult::GenerateStub;
+		}
+		Struct = Struct->GetSuperStruct();
+	}
 
-	FCookCommandParams CookParams(FCommandLine::Get());
-	FNativeCodeGenCommandlineParams CommandLineParams(CookParams.ToConversionParams());
-	check(FPaths::IsSamePath(CommandLineParams.ManifestFilePath, CommandLineParams.ManifestFilePath));
-	check(FBlueprintNativeCodeGenUtils::FinalizePlugin(GetManifest(), CommandLineParams));
+	return Result;
 }
 
 IMPLEMENT_MODULE(FBlueprintNativeCodeGenModule, BlueprintNativeCodeGen);
\ No newline at end of file
diff --git a/Engine/Source/Developer/BlueprintNativeCodeGen/Private/BlueprintNativeCodeGenUtils.cpp b/Engine/Source/Developer/BlueprintNativeCodeGen/Private/BlueprintNativeCodeGenUtils.cpp
--- a/Engine/Source/Developer/BlueprintNativeCodeGen/Private/BlueprintNativeCodeGenUtils.cpp
+++ b/Engine/Source/Developer/BlueprintNativeCodeGen/Private/BlueprintNativeCodeGenUtils.cpp
@@ -3,7 +3,6 @@
 #include "BlueprintNativeCodeGenPCH.h"
 #include "BlueprintNativeCodeGenManifest.h"
 #include "BlueprintNativeCodeGenUtils.h"
-#include "NativeCodeGenCommandlineParams.h"
 #include "Kismet2/KismetReinstanceUtilities.h"	 // for FBlueprintCompileReinstancer
 #include "Kismet2/CompilerResultsLog.h" 
 #include "KismetCompilerModule.h"
@@ -230,12 +229,13 @@ static UClass* BlueprintNativeCodeGenUtilsImpl::ResolveReplacementType(const FCo
  ******************************************************************************/
 
 //------------------------------------------------------------------------------
-bool FBlueprintNativeCodeGenUtils::FinalizePlugin(const FBlueprintNativeCodeGenManifest& Manifest, const FNativeCodeGenCommandlineParams& CommandParams)
+bool FBlueprintNativeCodeGenUtils::FinalizePlugin(const FBlueprintNativeCodeGenManifest& Manifest)
 {
 	bool bSuccess = true;
+	FBlueprintNativeCodeGenPaths TargetPaths = Manifest.GetTargetPaths();
 	bSuccess = bSuccess && BlueprintNativeCodeGenUtilsImpl::GenerateModuleBuildFile(Manifest);
-	bSuccess = bSuccess && BlueprintNativeCodeGenUtilsImpl::GenerateModuleSourceFiles(Manifest.GetTargetPaths());
-	bSuccess = bSuccess && BlueprintNativeCodeGenUtilsImpl::GeneratePluginDescFile(CommandParams.PluginName, Manifest.GetTargetPaths());
+	bSuccess = bSuccess && BlueprintNativeCodeGenUtilsImpl::GenerateModuleSourceFiles(TargetPaths);
+	bSuccess = bSuccess && BlueprintNativeCodeGenUtilsImpl::GeneratePluginDescFile(TargetPaths.RuntimeModuleName(), TargetPaths);
 	return bSuccess;
 }
 
@@ -268,7 +268,7 @@ void FBlueprintNativeCodeGenUtils::GenerateCppCode(UObject* Obj, TSharedPtr<FStr
 
 		FDisableGatheringDataOnScope DisableFib;
 
-		const FString TempPackageName = FString::Printf(TEXT("/Temp/__TEMP_BP__/%s"), *InBlueprintObj->GetName());
+		const FString TempPackageName = FString::Printf(TEXT("/Temp/__TEMP_BP__%s"), *InBlueprintObj->GetOutermost()->GetPathName());
 		UPackage* TempPackage = CreatePackage(nullptr, *TempPackageName);
 		check(TempPackage);
 		ON_SCOPE_EXIT
@@ -279,8 +279,9 @@ void FBlueprintNativeCodeGenUtils::GenerateCppCode(UObject* Obj, TSharedPtr<FStr
 
 		UBlueprint* DuplicateBP = nullptr;
 		{
-			FBlueprintDuplicationScopeFlags BPDuplicationFlags(
-				FBlueprintDuplicationScopeFlags::NoExtraCompilation | FBlueprintDuplicationScopeFlags::TheSameTimelineGuid);
+			FBlueprintDuplicationScopeFlags BPDuplicationFlags(FBlueprintDuplicationScopeFlags::NoExtraCompilation 
+				| FBlueprintDuplicationScopeFlags::TheSameTimelineGuid 
+				| FBlueprintDuplicationScopeFlags::ValidatePinsUsingSourceClass);
 			DuplicateBP = DuplicateObject<UBlueprint>(InBlueprintObj, TempPackage, *InBlueprintObj->GetName());
 		}
 		ensure((nullptr != DuplicateBP->GeneratedClass) && (InBlueprintObj->GeneratedClass != DuplicateBP->GeneratedClass));
diff --git a/Engine/Source/Developer/BlueprintNativeCodeGen/Private/BlueprintNativeCodeGenUtils.h b/Engine/Source/Developer/BlueprintNativeCodeGen/Private/BlueprintNativeCodeGenUtils.h
--- a/Engine/Source/Developer/BlueprintNativeCodeGen/Private/BlueprintNativeCodeGenUtils.h
+++ b/Engine/Source/Developer/BlueprintNativeCodeGen/Private/BlueprintNativeCodeGenUtils.h
@@ -5,7 +5,6 @@
 #include "FeedbackContext.h"
 
 // Forward declares
-struct FNativeCodeGenCommandlineParams;
 class  SBuildProgressWidget;
 struct FBlueprintNativeCodeGenManifest;
 
@@ -19,9 +18,8 @@ public:
 	 * Generated module build file, module source files, and plugin description file for the provided manifest.
 	 *
 	 * @param	Manifest		The manifest containing lists of converted files, etc
-	 * @param	CommandParams	Command parameters used to generate native code, includes plugin name
 	 */
-	static bool FinalizePlugin(const FBlueprintNativeCodeGenManifest& Manifest, const FNativeCodeGenCommandlineParams& CommandParams);
+	static bool FinalizePlugin(const FBlueprintNativeCodeGenManifest& Manifest);
 
 	/**
 	 * Recompiles the bytecode of a blueprint only. Should only be run for
diff --git a/Engine/Source/Developer/BlueprintNativeCodeGen/Private/NativeCodeGenCommandlineParams.cpp b/Engine/Source/Developer/BlueprintNativeCodeGen/Private/NativeCodeGenCommandlineParams.cpp
deleted file mode 100644
diff --git a/Engine/Source/Developer/BlueprintNativeCodeGen/Private/NativeCodeGenCommandlineParams.h b/Engine/Source/Developer/BlueprintNativeCodeGen/Private/NativeCodeGenCommandlineParams.h
deleted file mode 100644
diff --git a/Engine/Source/Developer/BlueprintNativeCodeGen/Public/BlueprintNativeCodeGenModule.h b/Engine/Source/Developer/BlueprintNativeCodeGen/Public/BlueprintNativeCodeGenModule.h
--- a/Engine/Source/Developer/BlueprintNativeCodeGen/Public/BlueprintNativeCodeGenModule.h
+++ b/Engine/Source/Developer/BlueprintNativeCodeGen/Public/BlueprintNativeCodeGenModule.h
@@ -7,46 +7,33 @@
 #include "ModuleManager.h"	// for GetModuleChecked<>
 #include "Blueprint/BlueprintSupport.h"
 
-#include "BlueprintNativeCodeGenModule.generated.h"
-
-struct BLUEPRINTNATIVECODEGEN_API FCookCommandParams
+struct FNativeCodeGenInitData
 {
-	FCookCommandParams(const FString& CommandLine);
-
-	bool    bRunConversion;
-	FString ManifestFilePath;
-	FString OutputPath;
-	FString PluginName;
+	// This is an array of pairs, the pairs are PlatformName/PlatformTargetDirectory. These
+	// are determined by the cooker:
+	TArray< TPair< FString, FString > > CodegenTargets;
 
-	TArray<FString> ToConversionParams() const;
+	// Optional Manifest ManifestIdentifier, used for child cook processes that need a unique manifest name.
+	// The identifier is used to make a unique name for each platform that is converted.
+	int32 ManifestIdentifier;
 };
 
-/*******************************************************************************
-* UBlueprintNativeCodeGenConfig
-******************************************************************************/
-UCLASS(config = Editor)
-class BLUEPRINTNATIVECODEGEN_API UBlueprintNativeCodeGenConfig : public UObject
-{
-	GENERATED_UCLASS_BODY()
-
-public:
-	UPROPERTY(globalconfig)
-	TArray<FString> PackagesToNeverConvert;
-
-	UPROPERTY(globalconfig)
-	TArray<FString> ExcludedAssetTypes;
-
-	UPROPERTY(globalconfig)
-	TArray<FString> ExcludedBlueprintTypes;
-};
-
-
 /** 
  * 
  */
 class IBlueprintNativeCodeGenModule : public IModuleInterface
 {
 public:
+	FORCEINLINE static void InitializeModule(const FNativeCodeGenInitData& InitData);
+
+	/**
+	* Utility function to reconvert all assets listed in a manifest, used to make fixes to
+	* the code generator itself and quickly test them with an already converted project.
+	*
+	* Not for use with any kind of incremental cooking.
+	*/
+	FORCEINLINE static void InitializeModuleForRerunDebugOnly(const TArray< TPair< FString, FString > >& CodegenTargets);
+
 	/**
 	 * Wrapper function that retrieves the interface to this module from the 
 	 * module-manager (so we can keep dependent code free of hardcoded strings,
@@ -58,6 +45,11 @@ public:
 	{
 		return FModuleManager::LoadModuleChecked<IBlueprintNativeCodeGenModule>(GetModuleName());
 	}
+
+	static bool IsNativeCodeGenModuleLoaded()
+	{
+		return FModuleManager::Get().IsModuleLoaded(GetModuleName());
+	}
 	
 	/**
 	 * Creates a centralized point where the name of this module is supplied 
@@ -71,8 +63,24 @@ public:
 		return TEXT("BlueprintNativeCodeGen");
 	}
 
-	virtual void Convert(UPackage* Package, EReplacementResult ReplacementType) = 0;
-	virtual void SaveManifest(const TCHAR* Filename) = 0;
-	virtual void MergeManifest(const TCHAR* Filename) = 0;
+	virtual void Convert(UPackage* Package, ESavePackageResult ReplacementType, const TCHAR* PlatformName) = 0;
+	virtual void SaveManifest(int32 Id = -1) = 0;
+	virtual void MergeManifest(int32 ManifestIdentifier) = 0;
 	virtual void FinalizeManifest() = 0;
+protected:
+	virtual void Initialize(const FNativeCodeGenInitData& InitData) = 0;
+	virtual void InitializeForRerunDebugOnly(const TArray< TPair< FString, FString > >& CodegenTargets) = 0;
 };
+
+void IBlueprintNativeCodeGenModule::InitializeModule(const FNativeCodeGenInitData& InitData)
+{
+	IBlueprintNativeCodeGenModule& Module = FModuleManager::LoadModuleChecked<IBlueprintNativeCodeGenModule>(GetModuleName());
+	Module.Initialize(InitData);
+}
+
+void IBlueprintNativeCodeGenModule::InitializeModuleForRerunDebugOnly(const TArray< TPair< FString, FString > >& CodegenTargets)
+{
+	IBlueprintNativeCodeGenModule& Module = FModuleManager::LoadModuleChecked<IBlueprintNativeCodeGenModule>(GetModuleName());
+	Module.InitializeForRerunDebugOnly(CodegenTargets);
+}
+
diff --git a/Engine/Source/Editor/BlueprintGraph/Private/EdGraphSchema_K2.cpp b/Engine/Source/Editor/BlueprintGraph/Private/EdGraphSchema_K2.cpp
--- a/Engine/Source/Editor/BlueprintGraph/Private/EdGraphSchema_K2.cpp
+++ b/Engine/Source/Editor/BlueprintGraph/Private/EdGraphSchema_K2.cpp
@@ -2013,7 +2013,7 @@ const FPinConnectionResponse UEdGraphSchema_K2::DetermineConnectionResponseOfCom
 	}
 }
 
-static FText GetPinIncompatibilityMessage(const UEdGraphPin* PinA, const UEdGraphPin* PinB)
+static FText GetPinIncompatibilityReason(const UEdGraphPin* PinA, const UEdGraphPin* PinB, bool* bIsFatalOut = nullptr)
 {
 	const FEdGraphPinType& PinAType = PinA->PinType;
 	const FEdGraphPinType& PinBType = PinB->PinType;
@@ -2043,6 +2043,11 @@ static FText GetPinIncompatibilityMessage(const UEdGraphPin* PinA, const UEdGrap
 			{
 				MessageFormat = LOCTEXT("ChildStructIncompatible", "Only exactly matching structures are considered compatible. Derived structures are disallowed.");
 			}
+
+			if (bIsFatalOut != nullptr)
+			{
+				*bIsFatalOut = true;
+			}
 		}
 	}
 	else if (OutputType.PinCategory == UEdGraphSchema_K2::PC_Class)
@@ -2148,8 +2153,15 @@ const FPinConnectionResponse UEdGraphSchema_K2::CanCreateConnection(const UEdGra
 		}
 		else
 		{
-			FText IncompatibilityReasonText = GetPinIncompatibilityMessage(PinA, PinB);
-			return FPinConnectionResponse(CONNECT_RESPONSE_DISALLOW, IncompatibilityReasonText.ToString());
+			bool bIsFatal = false;
+			FText IncompatibilityReasonText = GetPinIncompatibilityReason(PinA, PinB, &bIsFatal);
+
+			FPinConnectionResponse ConnectionResponse(CONNECT_RESPONSE_DISALLOW, IncompatibilityReasonText.ToString());
+			if (bIsFatal)
+			{
+				ConnectionResponse.SetFatal();
+			}
+			return ConnectionResponse;
 		}
 	}
 }
@@ -2342,7 +2354,9 @@ bool UEdGraphSchema_K2::SearchForAutocastFunction(const UEdGraphPin* OutputPin,
 		bool const bInputIsUObject = ((InputClass != NULL) && (InputClass == UObject::StaticClass()));
 		if (bInputIsUObject)
 		{
-			TargetFunction = TEXT("Conv_InterfaceToObject");
+			UFunction* Function = UKismetSystemLibrary::StaticClass()->FindFunctionByName(GET_MEMBER_NAME_CHECKED(UKismetSystemLibrary, Conv_InterfaceToObject));
+			TargetFunction = Function->GetFName();
+			FunctionOwner = Function->GetOwnerClass();
 		}
 	}
 	else if (OutputPin->PinType.PinCategory == PC_Object)
@@ -2355,12 +2369,16 @@ bool UEdGraphSchema_K2::SearchForAutocastFunction(const UEdGraphPin* OutputPin,
 				(InputClass != nullptr) &&
 				OutputClass->IsChildOf(InputClass))
 			{
-				TargetFunction = TEXT("GetObjectClass");
+				UFunction* Function = UGameplayStatics::StaticClass()->FindFunctionByName(GET_MEMBER_NAME_CHECKED(UGameplayStatics, GetObjectClass));
+				TargetFunction = Function->GetFName();
+				FunctionOwner = Function->GetOwnerClass();
 			}
 		}
 		else if (InputPin->PinType.PinCategory == PC_String)
 		{
-			TargetFunction = TEXT("GetDisplayName");
+			UFunction* Function = UKismetSystemLibrary::StaticClass()->FindFunctionByName(GET_MEMBER_NAME_CHECKED(UKismetSystemLibrary, GetDisplayName));
+			TargetFunction = Function->GetFName();
+			FunctionOwner = Function->GetOwnerClass();
 		}
 	}
 	else if (OutputPin->PinType.PinCategory == PC_Struct)
@@ -2371,7 +2389,9 @@ bool UEdGraphSchema_K2::SearchForAutocastFunction(const UEdGraphPin* OutputPin,
 			const UScriptStruct* InputStructType = Cast<const UScriptStruct>(InputPin->PinType.PinSubCategoryObject.Get());
 			if ((InputPin->PinType.PinCategory == PC_Struct) && (InputStructType == TBaseStructure<FTransform>::Get()))
 			{
-				TargetFunction = TEXT("MakeTransform");
+				UFunction* Function = UKismetMathLibrary::StaticClass()->FindFunctionByName(GET_MEMBER_NAME_CHECKED(UKismetMathLibrary, MakeTransform));
+				TargetFunction = Function->GetFName();
+				FunctionOwner = Function->GetOwnerClass();
 			}
 		}
 	}
@@ -3696,6 +3716,59 @@ bool UEdGraphSchema_K2::ArePinsCompatible(const UEdGraphPin* PinA, const UEdGrap
 	}
 }
 
+namespace
+{
+	static UClass* GetOriginalClassToFixCompatibilit(const UClass* InClass)
+	{
+		const UBlueprint* BP = InClass ? Cast<const UBlueprint>(InClass->ClassGeneratedBy) : nullptr;
+		return BP ? Cast<UClass>(BP->OriginalClass) : nullptr;
+	}
+
+	static bool ExtendedIsChildOf(const UClass* Child, const UClass* Parent)
+	{
+		if (Child->IsChildOf(Parent))
+		{
+			return true;
+		}
+
+		const UClass* OriginalChild = GetOriginalClassToFixCompatibilit(Child);
+		if (OriginalChild && OriginalChild->IsChildOf(Parent))
+		{
+			return true;
+		}
+
+		const UClass* OriginalParent = GetOriginalClassToFixCompatibilit(Parent);
+		if (OriginalParent && Child->IsChildOf(OriginalParent))
+		{
+			return true;
+		}
+
+		return false;
+	}
+
+	static bool ExtendedImplementsInterface(const UClass* Class, const UClass* Interface)
+	{
+		if (Class->ImplementsInterface(Interface))
+		{
+			return true;
+		}
+
+		const UClass* OriginalClass = GetOriginalClassToFixCompatibilit(Class);
+		if (OriginalClass && OriginalClass->ImplementsInterface(Interface))
+		{
+			return true;
+		}
+
+		const UClass* OriginalInterface = GetOriginalClassToFixCompatibilit(Interface);
+		if (OriginalInterface && Class->ImplementsInterface(OriginalInterface))
+		{
+			return true;
+		}
+
+		return false;
+	}
+};
+
 bool UEdGraphSchema_K2::ArePinTypesCompatible(const FEdGraphPinType& Output, const FEdGraphPinType& Input, const UClass* CallingContext, bool bIgnoreArray /*= false*/) const
 {
 	if( !bIgnoreArray && (Output.bIsArray != Input.bIsArray) && (Input.PinCategory != PC_Wildcard || Input.bIsArray) )
@@ -3718,7 +3791,7 @@ bool UEdGraphSchema_K2::ArePinTypesCompatible(const FEdGraphPinType& Output, con
 			UClass const* InputClass = Cast<UClass const>(Input.PinSubCategoryObject.Get());
 			check(InputClass && InputClass->IsChildOf(UInterface::StaticClass()));
 			
-			return OutputClass->IsChildOf(InputClass);
+			return ExtendedIsChildOf(OutputClass, InputClass);
 		}
 		else if (((Output.PinCategory == PC_Asset) && (Input.PinCategory == PC_Asset))
 			|| ((Output.PinCategory == PC_AssetClass) && (Input.PinCategory == PC_AssetClass)))
@@ -3727,7 +3800,7 @@ bool UEdGraphSchema_K2::ArePinTypesCompatible(const FEdGraphPinType& Output, con
 			const UClass* InputObject = (Input.PinSubCategory == PSC_Self) ? CallingContext : Cast<const UClass>(Input.PinSubCategoryObject.Get());
 			if ((OutputObject != NULL) && (InputObject != NULL))
 			{
-				return OutputObject->IsChildOf(InputObject);
+				return ExtendedIsChildOf(OutputObject ,InputObject);
 			}
 		}
 		else if ((Output.PinCategory == PC_Object) || (Output.PinCategory == PC_Struct) || (Output.PinCategory == PC_Class))
@@ -3748,22 +3821,23 @@ bool UEdGraphSchema_K2::ArePinTypesCompatible(const FEdGraphPinType& Output, con
 				const bool bInputIsInterface  = InputObject->IsChildOf(UInterface::StaticClass());
 				const bool bOutputIsInterface = OutputObject->IsChildOf(UInterface::StaticClass());
 
+				UClass const* OutputClass = Cast<const UClass>(OutputObject);
+				UClass const* InputClass = Cast<const UClass>(InputObject);
+
 				if (bInputIsInterface != bOutputIsInterface) 
 				{
-					UClass const* OutputClass = Cast<const UClass>(OutputObject);
-					UClass const* InputClass  = Cast<const UClass>(InputObject);
-
 					if (bInputIsInterface && (OutputClass != NULL))
 					{
-						return OutputClass->ImplementsInterface(InputClass);
+						return ExtendedImplementsInterface(OutputClass, InputClass);
 					}
 					else if (bOutputIsInterface && (InputClass != NULL))
 					{
-						return InputClass->ImplementsInterface(OutputClass);
+						return ExtendedImplementsInterface(InputClass, OutputClass);
 					}
 				}				
 
-				return OutputObject->IsChildOf(InputObject) && (bInputIsInterface == bOutputIsInterface);
+				return (OutputObject->IsChildOf(InputObject) || (OutputClass && InputClass && ExtendedIsChildOf(OutputClass, InputClass)))
+					&& (bInputIsInterface == bOutputIsInterface);
 			}
 		}
 		else if ((Output.PinCategory == PC_Byte) && (Output.PinSubCategory == Input.PinSubCategory))
@@ -3837,7 +3911,7 @@ bool UEdGraphSchema_K2::ArePinTypesCompatible(const FEdGraphPinType& Output, con
 			OutputClass = CallingContext;
 		}
 
-		return OutputClass && (OutputClass->ImplementsInterface(InterfaceClass) || OutputClass->IsChildOf(InterfaceClass));
+		return OutputClass && (ExtendedImplementsInterface(OutputClass, InterfaceClass) || ExtendedIsChildOf(OutputClass, InterfaceClass));
 	}
 
 	// Pins representing BLueprint objects and subclass of UObject can match when EditoronlyBP.bAllowClassAndBlueprintPinMatching=true (BaseEngine.ini)
@@ -5463,30 +5537,6 @@ TSharedPtr<FEdGraphSchemaAction> UEdGraphSchema_K2::GetCreateCommentAction() con
 bool UEdGraphSchema_K2::CanDuplicateGraph(UEdGraph* InSourceGraph) const
 {
 	EGraphType GraphType = GetGraphType(InSourceGraph);
-
-	if(GraphType == GT_Function)
-	{
-		UBlueprint* SourceBP = FBlueprintEditorUtils::FindBlueprintForGraph(InSourceGraph);
-
-		// Do not duplicate graphs in Blueprint Interfaces
-		if(SourceBP->BlueprintType == BPTYPE_Interface)
-		{
-			return false;
-		}
-
-		// Do not duplicate functions from implemented interfaces
-		if( FBlueprintEditorUtils::FindFunctionInImplementedInterfaces(SourceBP, InSourceGraph->GetFName()) )
-		{
-			return false;
-		}
-
-		// Do not duplicate inherited functions
-		if( FindField<UFunction>(SourceBP->ParentClass, InSourceGraph->GetFName()) )
-		{
-			return false;
-		}
-	}
-
 	return GraphType == GT_Function || GraphType == GT_Macro;
 }
 
@@ -5502,6 +5552,41 @@ UEdGraph* UEdGraphSchema_K2::DuplicateGraph(UEdGraph* GraphToDuplicate) const
 
 		if (NewGraph)
 		{
+			bool bIsOverrideGraph = false;
+			if (Blueprint->BlueprintType == BPTYPE_Interface)
+			{
+				bIsOverrideGraph = true;
+			}
+			else if (FBlueprintEditorUtils::FindFunctionInImplementedInterfaces(Blueprint, GraphToDuplicate->GetFName()))
+			{
+				bIsOverrideGraph = true;
+			}
+			else if (FindField<UFunction>(Blueprint->ParentClass, GraphToDuplicate->GetFName()))
+			{
+				bIsOverrideGraph = true;
+			}
+
+			// When duplicating an override function, we must put the graph through some extra work to properly own the data being duplicated, instead of expecting pin information will come from a parent
+			if (bIsOverrideGraph)
+			{
+				FBlueprintEditorUtils::PromoteGraphFromInterfaceOverride(Blueprint, NewGraph);
+				
+				// Remove all calls to the parent function, fix any exec pin links to pass through
+				TArray< UK2Node_CallParentFunction* > ParentFunctionCalls;
+				NewGraph->GetNodesOfClass(ParentFunctionCalls);
+
+				for (UK2Node_CallParentFunction* ParentFunctionCall : ParentFunctionCalls)
+				{
+					UEdGraphPin* ExecPin = ParentFunctionCall->GetExecPin();
+					UEdGraphPin* ThenPin = ParentFunctionCall->GetThenPin();
+					if (ExecPin->LinkedTo.Num() && ThenPin->LinkedTo.Num())
+					{
+						MovePinLinks(*ExecPin, *ThenPin->LinkedTo[0]);
+					}
+					NewGraph->RemoveNode(ParentFunctionCall);
+				}
+			}
+
 			FName NewGraphName = FBlueprintEditorUtils::FindUniqueKismetName(Blueprint, GraphToDuplicate->GetFName().GetPlainNameString());
 			FEdGraphUtilities::RenameGraphCloseToName(NewGraph,NewGraphName.ToString());
 			// can't have two graphs with the same guid... that'd be silly!
@@ -5527,6 +5612,9 @@ UEdGraph* UEdGraphSchema_K2::DuplicateGraph(UEdGraph* GraphToDuplicate) const
 					CustomEvent->RenameCustomEventCloseToName();
 				}
 			}
+
+			// Potentially adjust variable names for any child blueprints
+			FBlueprintEditorUtils::ValidateBlueprintChildVariables(Blueprint, NewGraph->GetFName());
 		}
 	}
 	return NewGraph;
diff --git a/Engine/Source/Editor/BlueprintGraph/Private/K2Node.cpp b/Engine/Source/Editor/BlueprintGraph/Private/K2Node.cpp
--- a/Engine/Source/Editor/BlueprintGraph/Private/K2Node.cpp
+++ b/Engine/Source/Editor/BlueprintGraph/Private/K2Node.cpp
@@ -365,6 +365,9 @@ void UK2Node::ReconstructNode()
 
 	UBlueprint* Blueprint = GetBlueprint();
 
+	FLinkerLoad* Linker = Blueprint->GetLinker();
+	const UEdGraphSchema* Schema = GetSchema();
+
 	// Break any links to 'orphan' pins
 	for (int32 PinIndex = 0; PinIndex < Pins.Num(); ++PinIndex)
 	{
@@ -378,6 +381,19 @@ void UK2Node::ReconstructNode()
 			{
 				Pin->LinkedTo.Remove(OtherPin);
 			}
+
+			if (Blueprint->bIsRegeneratingOnLoad && Linker->UE4Ver() < VER_UE4_INJECT_BLUEPRINT_STRUCT_PIN_CONVERSION_NODES)
+			{
+				if ((Pin->PinType.PinCategory != UEdGraphSchema_K2::PC_Struct))
+				{
+					continue;
+				}
+
+				if (Schema->CanCreateConnection(Pin, OtherPin).Response == ECanCreateConnectionResponse::CONNECT_RESPONSE_MAKE_WITH_CONVERSION_NODE)
+				{
+					Schema->CreateAutomaticConversionNodeAndConnections(Pin, OtherPin);
+				}
+			}
 		}
 	}
 
diff --git a/Engine/Source/Editor/BlueprintGraph/Private/K2Node_CustomEvent.cpp b/Engine/Source/Editor/BlueprintGraph/Private/K2Node_CustomEvent.cpp
--- a/Engine/Source/Editor/BlueprintGraph/Private/K2Node_CustomEvent.cpp
+++ b/Engine/Source/Editor/BlueprintGraph/Private/K2Node_CustomEvent.cpp
@@ -261,6 +261,15 @@ uint32 UK2Node_CustomEvent::GetNetFlags() const
 		// inherited net flags take precedence 
 		NetFlags = (ParentFunction->FunctionFlags & FUNC_NetFuncFlags);
 	}
+
+	// Sanitize NetFlags, only allow replication flags that can be supported by the online system
+	// This mirrors logic in ProcessFunctionSpecifiers in HeaderParser.cpp. Basically if we want to 
+	// replicate a function we need to know whether we're replicating on the client or the server.
+	if (!(NetFlags & FUNC_Net))
+	{
+		NetFlags = 0;
+	}
+
 	return NetFlags;
 }
 
diff --git a/Engine/Source/Editor/BlueprintGraph/Private/K2Node_EaseFunction.cpp b/Engine/Source/Editor/BlueprintGraph/Private/K2Node_EaseFunction.cpp
--- a/Engine/Source/Editor/BlueprintGraph/Private/K2Node_EaseFunction.cpp
+++ b/Engine/Source/Editor/BlueprintGraph/Private/K2Node_EaseFunction.cpp
@@ -426,7 +426,7 @@ void UK2Node_EaseFunction::ExpandNode(class FKismetCompilerContext& CompilerCont
 		CompilerContext.MovePinLinksToIntermediate(*StepsPin, *CallFunction->FindPinChecked(FEaseFunctionNodeHelper::GetStepsPinName()));
 	}
 
-	// Cleanup links to ourselfs and we are done!
+	// Cleanup links to ourself and we are done!
 	BreakAllNodeLinks();
 }
 
@@ -451,7 +451,8 @@ void UK2Node_EaseFunction::RefreshPinVisibility()
 
 	UEdGraphPin* BlendExpPin = FindPinChecked(FEaseFunctionNodeHelper::GetBlendExpPinName());
 
-	if (NewEasingFunc == EEasingFunc::EaseIn ||
+	if (NewEasingFunc == -1 ||
+		NewEasingFunc == EEasingFunc::EaseIn ||
 		NewEasingFunc == EEasingFunc::EaseOut ||
 		NewEasingFunc == EEasingFunc::EaseInOut)
 	{
@@ -466,7 +467,8 @@ void UK2Node_EaseFunction::RefreshPinVisibility()
 	}
 
 	UEdGraphPin* StepsPin = FindPinChecked(FEaseFunctionNodeHelper::GetStepsPinName());
-	if (NewEasingFunc == EEasingFunc::Step)
+	if (NewEasingFunc == -1 || 
+		NewEasingFunc == EEasingFunc::Step)
 	{
 		// Show the Steps pin:
 		StepsPin->bHidden = false;
diff --git a/Engine/Source/Editor/BlueprintGraph/Private/K2Node_FunctionResult.cpp b/Engine/Source/Editor/BlueprintGraph/Private/K2Node_FunctionResult.cpp
--- a/Engine/Source/Editor/BlueprintGraph/Private/K2Node_FunctionResult.cpp
+++ b/Engine/Source/Editor/BlueprintGraph/Private/K2Node_FunctionResult.cpp
@@ -16,7 +16,7 @@ private:
 		{
 			const bool bValuePin = (Pin->PinType.PinCategory != K2Schema->PC_Exec);
 			const bool bNotConnected = (Pin->Direction == EEdGraphPinDirection::EGPD_Input) && (0 == Pin->LinkedTo.Num());
-			const bool bNeedToResetDefaultValue = !(K2Schema->IsPinDefaultValid(Pin, Pin->DefaultValue, Pin->DefaultObject, Pin->DefaultTextValue).IsEmpty());
+			const bool bNeedToResetDefaultValue = (Pin->DefaultValue.IsEmpty() && Pin->DefaultObject == nullptr && Pin->DefaultTextValue.IsEmpty()) || !(K2Schema->IsPinDefaultValid(Pin, Pin->DefaultValue, Pin->DefaultObject, Pin->DefaultTextValue).IsEmpty());
 			if (bValuePin && bNotConnected && bNeedToResetDefaultValue)
 			{
 				K2Schema->SetPinDefaultValueBasedOnType(Pin);
diff --git a/Engine/Source/Editor/BlueprintGraph/Private/K2Node_FunctionTerminator.cpp b/Engine/Source/Editor/BlueprintGraph/Private/K2Node_FunctionTerminator.cpp
--- a/Engine/Source/Editor/BlueprintGraph/Private/K2Node_FunctionTerminator.cpp
+++ b/Engine/Source/Editor/BlueprintGraph/Private/K2Node_FunctionTerminator.cpp
@@ -91,7 +91,8 @@ void UK2Node_FunctionTerminator::PromoteFromInterfaceOverride(bool bIsPrimaryTer
 {
 	// Remove the signature class, that is not relevant.
 	SignatureClass = nullptr;
-	for (const UEdGraphPin* Pin : Pins)
+	TArray<UEdGraphPin*> OriginalPins = Pins;
+	for (const UEdGraphPin* Pin : OriginalPins)
 	{
 		if (Pin->PinType.PinCategory != UEdGraphSchema_K2::PC_Exec)
 		{
diff --git a/Engine/Source/Editor/GraphEditor/Private/SGraphEditorImpl.cpp b/Engine/Source/Editor/GraphEditor/Private/SGraphEditorImpl.cpp
--- a/Engine/Source/Editor/GraphEditor/Private/SGraphEditorImpl.cpp
+++ b/Engine/Source/Editor/GraphEditor/Private/SGraphEditorImpl.cpp
@@ -757,6 +757,11 @@ FSlateColor SGraphEditorImpl::InstructionBorderColor() const
 	return BorderColor;
 }
 
+void SGraphEditorImpl::CaptureKeyboard()
+{
+	FSlateApplication::Get().SetKeyboardFocus(GraphPanel);
+}
+
 /////////////////////////////////////////////////////
 
 #undef LOCTEXT_NAMESPACE 
diff --git a/Engine/Source/Editor/GraphEditor/Private/SGraphEditorImpl.h b/Engine/Source/Editor/GraphEditor/Private/SGraphEditorImpl.h
--- a/Engine/Source/Editor/GraphEditor/Private/SGraphEditorImpl.h
+++ b/Engine/Source/Editor/GraphEditor/Private/SGraphEditorImpl.h
@@ -144,6 +144,7 @@ public:
 	virtual void SetPinVisibility(SGraphEditor::EPinVisibility Visibility) override;
 	virtual void StraightenConnections() override;
 	virtual void StraightenConnections(UEdGraphPin* SourcePin, UEdGraphPin* PinToAlign) override;
+	virtual void CaptureKeyboard() override;
 	// End of SGraphEditor interface
 protected:
 	//
diff --git a/Engine/Source/Editor/GraphEditor/Private/SGraphPanel.cpp b/Engine/Source/Editor/GraphEditor/Private/SGraphPanel.cpp
--- a/Engine/Source/Editor/GraphEditor/Private/SGraphPanel.cpp
+++ b/Engine/Source/Editor/GraphEditor/Private/SGraphPanel.cpp
@@ -1603,43 +1603,61 @@ void SGraphPanel::OnGraphChanged(const FEdGraphEditAction& EditAction)
 		// that the timer system requires (and we don't leverage):
 		if (bWasRemoveAction)
 		{
-			const auto RemoveNodeDelegateWrapper = [](double, float, SGraphPanel* Parent, const UEdGraphNode* Node) -> EActiveTimerReturnType
+			const auto RemoveNodeDelegateWrapper = [](double, float, SGraphPanel* Parent, TWeakObjectPtr<UEdGraphNode> NodePtr) -> EActiveTimerReturnType
 			{
-				Parent->RemoveNode(Node);
+				if (UEdGraphNode* Node = NodePtr.Get())
+				{
+					Parent->RemoveNode(Node);
+				}
 				return EActiveTimerReturnType::Stop;
 			};
 
-			for (auto Node : EditAction.Nodes)
+			for (const UEdGraphNode* Node : EditAction.Nodes)
 			{
-				RegisterActiveTimer(0.f, FWidgetActiveTimerDelegate::CreateStatic(RemoveNodeDelegateWrapper, this, Node));
+				TWeakObjectPtr<UEdGraphNode> NodePtr = Node;
+				RegisterActiveTimer(0.f, FWidgetActiveTimerDelegate::CreateStatic(RemoveNodeDelegateWrapper, this, NodePtr));
 			}
 		}
 		if (bWasAddAction)
 		{
-			const auto AddNodeDelegateWrapper = [](double, float, SGraphPanel* Parent, UEdGraphNode* Node, bool bForceUserAdded) -> EActiveTimerReturnType
+			const auto AddNodeDelegateWrapper = [](double, float, SGraphPanel* Parent, TWeakObjectPtr<UEdGraphNode> NodePtr, bool bForceUserAdded) -> EActiveTimerReturnType
 			{
-				Parent->RemoveNode(Node);
-				Parent->AddNode(Node, bForceUserAdded ? WasUserAdded : NotUserAdded );
+				if (UEdGraphNode* Node = NodePtr.Get())
+				{
+					Parent->RemoveNode(Node);
+					Parent->AddNode(Node, bForceUserAdded ? WasUserAdded : NotUserAdded);
+				}
 				return EActiveTimerReturnType::Stop;
 			};
 
-			for (auto Node : EditAction.Nodes)
+			for (const UEdGraphNode* Node : EditAction.Nodes)
 			{
-				RegisterActiveTimer(0.f, FWidgetActiveTimerDelegate::CreateStatic(AddNodeDelegateWrapper, this, const_cast<UEdGraphNode*>(Node), EditAction.bUserInvoked ) );
+				TWeakObjectPtr<UEdGraphNode> NodePtr = Node;
+				RegisterActiveTimer(0.f, FWidgetActiveTimerDelegate::CreateStatic(AddNodeDelegateWrapper, this, NodePtr, EditAction.bUserInvoked));
 			}
 		}
 		if (bWasSelectAction)
 		{
-			const auto SelectNodeDelegateWrapper = [](double, float, SGraphPanel* Parent, TSet<const UEdGraphNode*> Nodes) -> EActiveTimerReturnType
+			const auto SelectNodeDelegateWrapper = [](double, float, SGraphPanel* Parent, TSet< TWeakObjectPtr<UEdGraphNode> > NodePtrs) -> EActiveTimerReturnType
 			{
 				Parent->DeferredSelectionTargetObjects.Empty();
-				for (auto Node : Nodes)
+				for (TWeakObjectPtr<UEdGraphNode>& NodePtr : NodePtrs)
 				{
-					Parent->DeferredSelectionTargetObjects.Add(Node);
+					if (UEdGraphNode* Node = NodePtr.Get())
+					{
+						Parent->DeferredSelectionTargetObjects.Add(Node);
+					}
 				}
 				return EActiveTimerReturnType::Stop;
 			};
-			RegisterActiveTimer(0.f, FWidgetActiveTimerDelegate::CreateStatic(SelectNodeDelegateWrapper, this, EditAction.Nodes));
+
+			TSet< TWeakObjectPtr<UEdGraphNode> > NodePtrSet;
+			for (const UEdGraphNode* Node : EditAction.Nodes)
+			{
+				NodePtrSet.Add(Node);
+			}
+
+			RegisterActiveTimer(0.f, FWidgetActiveTimerDelegate::CreateStatic(SelectNodeDelegateWrapper, this, NodePtrSet));
 		}
 	}
 }
diff --git a/Engine/Source/Editor/Kismet/Private/BlueprintDetailsCustomization.cpp b/Engine/Source/Editor/Kismet/Private/BlueprintDetailsCustomization.cpp
--- a/Engine/Source/Editor/Kismet/Private/BlueprintDetailsCustomization.cpp
+++ b/Engine/Source/Editor/Kismet/Private/BlueprintDetailsCustomization.cpp
@@ -4307,6 +4307,13 @@ void FBlueprintInterfaceLayout::OnRemoveInterface(FInterfaceName InterfaceName)
 	UBlueprint* Blueprint = GlobalOptionsDetailsPtr.Pin()->GetBlueprintObj();
 	check(Blueprint);
 
+	const EAppReturnType::Type DialogReturn = FMessageDialog::Open(EAppMsgType::YesNoCancel, NSLOCTEXT("UnrealEd", "TransferInterfaceFunctionsToBlueprint", "Would you like to transfer the interface functions to be part of your blueprint?"));
+
+	if (DialogReturn == EAppReturnType::Cancel)
+	{
+		// We canceled!
+		return;
+	}
 	const FName InterfaceFName = InterfaceName.Name;
 
 	// Close all graphs that are about to be removed
@@ -4317,10 +4324,8 @@ void FBlueprintInterfaceLayout::OnRemoveInterface(FInterfaceName InterfaceName)
 		GlobalOptionsDetailsPtr.Pin()->GetBlueprintEditorPtr().Pin()->CloseDocumentTab(*GraphIt);
 	}
 
-	const bool bPreserveInterfaceFunctions = (EAppReturnType::Yes == FMessageDialog::Open( EAppMsgType::YesNo, NSLOCTEXT("UnrealEd", "TransferInterfaceFunctionsToBlueprint", "Would you like to transfer the interface functions to be part of your blueprint?") ));
-
 	// Do the work of actually removing the interface
-	FBlueprintEditorUtils::RemoveInterface(Blueprint, InterfaceFName, bPreserveInterfaceFunctions);
+	FBlueprintEditorUtils::RemoveInterface(Blueprint, InterfaceFName, DialogReturn == EAppReturnType::Yes);
 
 	RegenerateChildrenDelegate.ExecuteIfBound();
 
diff --git a/Engine/Source/Editor/Kismet/Private/BlueprintEditor.cpp b/Engine/Source/Editor/Kismet/Private/BlueprintEditor.cpp
--- a/Engine/Source/Editor/Kismet/Private/BlueprintEditor.cpp
+++ b/Engine/Source/Editor/Kismet/Private/BlueprintEditor.cpp
@@ -508,7 +508,10 @@ bool FBlueprintEditor::IsInAScriptingMode() const
 
 bool FBlueprintEditor::OnRequestClose()
 {
-	TabManager->InvokeTab(FBlueprintEditorTabs::BlueprintProfilerID)->RequestCloseTab();
+	if (GetDefault<UEditorExperimentalSettings>()->bBlueprintPerformanceAnalysisTools)
+	{
+		TabManager->InvokeTab(FBlueprintEditorTabs::BlueprintProfilerID)->RequestCloseTab();
+	}
 	bEditorMarkedAsClosed = true;
 	return FWorkflowCentricApplication::OnRequestClose();
 }
@@ -1252,6 +1255,9 @@ TSharedRef<SGraphEditor> FBlueprintEditor::CreateGraphEditorWidget(TSharedRef<FT
 		InEvents.OnCreateActionMenu = SGraphEditor::FOnCreateActionMenu::CreateSP(this, &FBlueprintEditor::OnCreateGraphActionMenu);
 	}
 
+	// Append play world commands
+	GraphEditorCommands->Append( FPlayWorldCommands::GlobalPlayWorldActions.ToSharedRef() );
+
 	TSharedRef<SGraphEditor> Editor = SNew(SGraphEditor)
 		.AdditionalCommands(GraphEditorCommands)
 		.IsEditable(this, &FBlueprintEditor::IsEditable, InGraph)
@@ -1674,7 +1680,12 @@ void FBlueprintEditor::InitBlueprintEditor(
 	PostLayoutBlueprintEditorInitialization();
 
 	// Ensure the profiler UI respects the current state if it had previously been enabled.
-	const bool bShowBlueprintProfilerTab = IsProfilerActive() && GetDefault<UEditorExperimentalSettings>()->bUnifiedBlueprintEditor;
+	bool bShowBlueprintProfilerTab = false;
+	if (GetDefault<UEditorExperimentalSettings>()->bBlueprintPerformanceAnalysisTools)
+	{
+		IBlueprintProfilerInterface& ProfilerModule = FModuleManager::LoadModuleChecked<IBlueprintProfilerInterface>("BlueprintProfiler");
+		bShowBlueprintProfilerTab = ProfilerModule.IsProfilerEnabled();
+	}
 	if (bShowBlueprintProfilerTab)
 	{
 		TabManager->InvokeTab(FBlueprintEditorTabs::BlueprintProfilerID);
@@ -2193,12 +2204,10 @@ void FBlueprintEditor::CreateDefaultTabContents(const TArray<UBlueprint*>& InBlu
 			. BlueprintToWatch(InBlueprint)
 			. IsEnabled(!bIsInterface && !bIsMacro);
 
-		if (GetDefault<UEditorExperimentalSettings>()->bBlueprintPerformanceAnalysisTools)
-		{
 			this->BlueprintProfiler =
 				SNew(SBlueprintProfilerView)
 				.ProfileViewType(PVT_LiveView);
-		}
+
 		this->Palette = 
 			SNew(SBlueprintPalette, SharedThis(this))
 				.IsEnabled(this, &FBlueprintEditor::IsFocusedGraphEditable);
@@ -7619,7 +7628,12 @@ TSharedPtr<SGraphEditor> FBlueprintEditor::OpenGraphAndBringToFront(UEdGraph* Gr
 	TSharedPtr<SDockTab> TabWithGraph = OpenDocument(Graph, FDocumentTracker::CreateHistoryEvent);
 
 	// We know that the contents of the opened tabs will be a graph editor.
-	return StaticCastSharedRef<SGraphEditor>(TabWithGraph->GetContent());
+	TSharedRef<SGraphEditor> NewGraphEditor = StaticCastSharedRef<SGraphEditor>(TabWithGraph->GetContent());
+
+	// Handover the keyboard focus to the new graph editor widget.
+	NewGraphEditor->CaptureKeyboard();
+
+	return NewGraphEditor;
 }
 
 TSharedPtr<SDockTab> FBlueprintEditor::OpenDocument(UObject* DocumentID, FDocumentTracker::EOpenDocumentCause Cause)
@@ -8252,21 +8266,6 @@ bool FBlueprintEditor::IsFocusedGraphEditable() const
 	return true;
 }
 
-void FBlueprintEditor::SaveAsset_Execute()
-{
-	for (auto Object : GetEditingObjects())
-	{
-		auto Blueprint = Cast<UBlueprint>(Object);
-		if (Blueprint)
-		{
-			// Update the Blueprint's search data
-			FFindInBlueprintSearchManager::Get().AddOrUpdateBlueprintSearchMetadata(Blueprint);
-		}
-	}
-
-	IBlueprintEditor::SaveAsset_Execute();
-}
-
 void FBlueprintEditor::TryInvokingDetailsTab(bool bFlash)
 {
 	if ( TabManager->CanSpawnTab(FBlueprintEditorTabs::DetailsID) )
diff --git a/Engine/Source/Editor/Kismet/Private/BlueprintEditorModes.cpp b/Engine/Source/Editor/Kismet/Private/BlueprintEditorModes.cpp
--- a/Engine/Source/Editor/Kismet/Private/BlueprintEditorModes.cpp
+++ b/Engine/Source/Editor/Kismet/Private/BlueprintEditorModes.cpp
@@ -112,10 +112,7 @@ FBlueprintEditorApplicationMode::FBlueprintEditorApplicationMode(TSharedPtr<clas
 	}
 	BlueprintEditorTabFactories.RegisterFactory(MakeShareable(new FCompilerResultsSummoner(InBlueprintEditor)));
 	BlueprintEditorTabFactories.RegisterFactory(MakeShareable(new FFindResultsSummoner(InBlueprintEditor)));
-	if (GetDefault<UEditorExperimentalSettings>()->bBlueprintPerformanceAnalysisTools)
-	{
-		BlueprintEditorTabFactories.RegisterFactory(MakeShareable(new FBlueprintProfilerSummoner(InBlueprintEditor)));
-	}
+	BlueprintEditorTabFactories.RegisterFactory(MakeShareable(new FBlueprintProfilerSummoner(InBlueprintEditor)));
 	
 	if( bRegisterViewport )
 	{
@@ -136,10 +133,7 @@ FBlueprintEditorApplicationMode::FBlueprintEditorApplicationMode(TSharedPtr<clas
 	InBlueprintEditor->GetToolbarBuilder()->AddScriptingToolbar(ToolbarExtender);
 	InBlueprintEditor->GetToolbarBuilder()->AddBlueprintGlobalOptionsToolbar(ToolbarExtender);
 	InBlueprintEditor->GetToolbarBuilder()->AddDebuggingToolbar(ToolbarExtender);
-	if (GetDefault<UEditorExperimentalSettings>()->bBlueprintPerformanceAnalysisTools)
-	{
-		InBlueprintEditor->GetToolbarBuilder()->AddProfilerToolbar(ToolbarExtender);
-	}
+	InBlueprintEditor->GetToolbarBuilder()->AddProfilerToolbar(ToolbarExtender);
 }
 
 void FBlueprintEditorApplicationMode::RegisterTabFactories(TSharedPtr<FTabManager> InTabManager)
@@ -581,10 +575,7 @@ FBlueprintEditorUnifiedMode::FBlueprintEditorUnifiedMode(TSharedPtr<class FBluep
 	}
 	BlueprintEditorTabFactories.RegisterFactory(MakeShareable(new FCompilerResultsSummoner(InBlueprintEditor)));
 	BlueprintEditorTabFactories.RegisterFactory(MakeShareable(new FFindResultsSummoner(InBlueprintEditor)));
-	if (GetDefault<UEditorExperimentalSettings>()->bBlueprintPerformanceAnalysisTools)
-	{
-		BlueprintEditorTabFactories.RegisterFactory(MakeShareable(new FBlueprintProfilerSummoner(InBlueprintEditor)));
-	}
+	BlueprintEditorTabFactories.RegisterFactory(MakeShareable(new FBlueprintProfilerSummoner(InBlueprintEditor)));
 	
 	if( bRegisterViewport )
 	{
diff --git a/Engine/Source/Editor/Kismet/Private/FindInBlueprintManager.cpp b/Engine/Source/Editor/Kismet/Private/FindInBlueprintManager.cpp
--- a/Engine/Source/Editor/Kismet/Private/FindInBlueprintManager.cpp
+++ b/Engine/Source/Editor/Kismet/Private/FindInBlueprintManager.cpp
@@ -313,6 +313,16 @@ namespace FiBSerializationHelpers
 
 namespace BlueprintSearchMetaDataHelpers
 {
+	/** Cache structure of searchable metadata and sub-properties relating to a Property */
+	struct FSearchableProperty
+	{
+		UProperty* TargetProperty;
+		bool bIsSearchableMD;
+		bool bIsShallowSearchableMD;
+		bool bIsMarkedNotSearchableMD;
+		TArray<FSearchableProperty> ChildProperties;
+	};
+
 	/** Json Writer used for serializing FText's in the correct format for Find-in-Blueprints */
 	template < class PrintPolicy = TPrettyJsonPrintPolicy<TCHAR> >
 	class TJsonFindInBlueprintStringWriter : public TJsonStringWriter<PrintPolicy>
@@ -428,10 +438,15 @@ namespace BlueprintSearchMetaDataHelpers
 
 			bool operator==(const FLookupTableItem& InObject) const
 			{
-				if(!Text.CompareTo(InObject.Text) &&
-					!FTextInspector::GetSourceString(Text)->Compare(*FTextInspector::GetDisplayString(InObject.Text), ESearchCase::CaseSensitive) )
+				if (!Text.CompareTo(InObject.Text))
 				{
-					return true;
+					if (FTextInspector::GetNamespace(Text).Get(TEXT("DefaultNamespace")) == FTextInspector::GetNamespace(InObject.Text).Get(TEXT("DefaultNamespace")))
+					{
+						if (FTextInspector::GetKey(Text).Get(TEXT("DefaultKey")) == FTextInspector::GetKey(InObject.Text).Get(TEXT("DefaultKey")))
+						{
+							return true;
+						}
+					}
 				}
 
 				return false;
@@ -456,14 +471,16 @@ namespace BlueprintSearchMetaDataHelpers
 			int32* TableLookupValuePtr = ReverseLookupTable.Find(FLookupTableItem(Text));
 			if(TableLookupValuePtr)
 			{
-				TJsonStringWriter<PrintPolicy>::WriteStringValue( FString::FromInt(*TableLookupValuePtr) );
+				TJsonStringWriter<PrintPolicy>::WriteStringValue(FString::FromInt(*TableLookupValuePtr));
 			}
 			else
 			{
 				// Add the FText to the table and write to the Json the ID to look the item up using
 				int32 TableLookupValue = LookupTable.Num();
-				LookupTable.Add(TableLookupValue, Text);
-				ReverseLookupTable.Add(FLookupTableItem(Text), TableLookupValue);
+				{
+					LookupTable.Add(TableLookupValue, Text);
+					ReverseLookupTable.Add(FLookupTableItem(Text), TableLookupValue);
+				}
 				TJsonStringWriter<PrintPolicy>::WriteStringValue( FString::FromInt(TableLookupValue) );
 			}
 		}
@@ -484,11 +501,18 @@ namespace BlueprintSearchMetaDataHelpers
 
 		// This is just locally needed for the write, to lookup the integer value by using the string of the FText
 		TMap< FLookupTableItem, int32 > ReverseLookupTable;
+
+	public:
+		/** Cached mapping of all searchable properties that have been discovered while gathering searchable data for the current Blueprint */
+		TMap<UStruct*, TArray<FSearchableProperty>> CachedPropertyMapping;
 	};
 
 	static uint32 GetTypeHash(const BlueprintSearchMetaDataHelpers::TJsonFindInBlueprintStringWriter<TCondensedJsonPrintPolicy<TCHAR>>::FLookupTableItem& InObject)
 	{
-		return FCrc::MemCrc32(&InObject.Text, sizeof(InObject.Text));
+		FString Namespace = FTextInspector::GetNamespace(InObject.Text).Get(TEXT("DefaultNamespace"));
+		FString Key = FTextInspector::GetKey(InObject.Text).Get(TEXT("DefaultKey"));
+		uint32 Hash = HashCombine(GetTypeHash(InObject.Text.ToString()), HashCombine(GetTypeHash(Namespace), GetTypeHash(Key)));
+		return Hash;
 	}
 
 	typedef TJsonWriter<TCHAR, TCondensedJsonPrintPolicy<TCHAR>> SearchMetaDataWriterParentClass;
@@ -744,29 +768,34 @@ namespace BlueprintSearchMetaDataHelpers
 				else
 				{
 					// Shallow conversion of property to string
-					auto JsonValue = FJsonObjectConverter::UPropertyToJsonValue(InProperty, InValue, 0, 0);
+					TSharedPtr<FJsonValue> JsonValue;
+					JsonValue = FJsonObjectConverter::UPropertyToJsonValue(InProperty, InValue, 0, 0);
 					FJsonSerializer::Serialize(JsonValue, InProperty->GetName(), StaticCastSharedRef<SearchMetaDataWriterParentClass>(InWriter), false);
 				}
 			}
 		}
 		else
 		{
-			auto JsonValue = FJsonObjectConverter::UPropertyToJsonValue(InProperty, InValue, 0, 0);
+			TSharedPtr<FJsonValue> JsonValue;
+			JsonValue = FJsonObjectConverter::UPropertyToJsonValue(InProperty, InValue, 0, 0);
 			FJsonSerializer::Serialize(JsonValue, InProperty->GetName(), StaticCastSharedRef<SearchMetaDataWriterParentClass>(InWriter), false);
 		}
 	}
 
-	void GatherSearchableProperties(TSharedRef< SearchMetaDataWriter>& InWriter, const void* InValue, UStruct* InStruct, EGatherSearchableType InSearchableType)
+	void GatherSearchableProperties(TSharedRef<SearchMetaDataWriter>& InWriter, const void* InValue, UStruct* InStruct, EGatherSearchableType InSearchableType)
 	{
 		if (InValue)
 		{
-			for (TFieldIterator<UProperty> PropIt(InStruct); PropIt; ++PropIt)
+			TArray<FSearchableProperty>* SearchablePropertyData = InWriter->CachedPropertyMapping.Find(InStruct);
+			check(SearchablePropertyData);
+
+			for (FSearchableProperty& SearchableProperty : *SearchablePropertyData)
 			{
-				UProperty* Property = *PropIt;
-				bool bIsSearchableMD = Property->GetBoolMetaData(*FFiBMD::FiBSearchableMD);
-				bool bIsShallowSearchableMD = Property->GetBoolMetaData(*FFiBMD::FiBSearchableShallowMD);
+				UProperty* Property = SearchableProperty.TargetProperty;
+				bool bIsSearchableMD = SearchableProperty.bIsSearchableMD;
+				bool bIsShallowSearchableMD = SearchableProperty.bIsShallowSearchableMD;
 				// It only is truly marked as not searchable if it has the metadata set to false, if the metadata is missing then we assume the searchable type that is passed in unless SEARCHABLE_AS_DESIRED
-				bool bIsMarkedNotSearchableMD = Property->HasMetaData(*FFiBMD::FiBSearchableMD) && !bIsSearchableMD;
+				bool bIsMarkedNotSearchableMD = SearchableProperty.bIsMarkedNotSearchableMD;
 
 				if ( (InSearchableType != SEARCHABLE_AS_DESIRED && !bIsMarkedNotSearchableMD) 
 					|| bIsShallowSearchableMD || bIsSearchableMD)
@@ -808,6 +837,113 @@ namespace BlueprintSearchMetaDataHelpers
 		}
 	}
 
+	/**
+	 * Caches all properties that have searchability metadata
+	 *
+	 * @param InOutCachePropertyMapping		Mapping of all the searchable properties that we are building
+	 * @param InValue						Value of the Object to serialize
+	 * @param InStruct						Struct or class that represent the UObject's layout
+	 * @param InSearchableType				Informs the system how it should examine the properties to determine if they are searchable. All sub-properties of searchable properties are automatically gathered unless marked as not being searchable
+	 */
+	void CacheSearchableProperties(TMap<UStruct*, TArray<FSearchableProperty>>& InOutCachePropertyMapping, const void* InValue, UStruct* InStruct, EGatherSearchableType InSearchableType = SEARCHABLE_AS_DESIRED);
+
+	/**
+	 * Digs into a property for any sub-properties that might exist so it can recurse and cache them
+	 *
+	 * @param InOutCachePropertyMapping		Mapping of all the searchable properties that we are building
+	 * @param InProperty					Property currently being cached
+	 * @param InValue						Value of the Object to serialize
+	 * @param InStruct						Struct or class that represent the UObject's layout
+	 */
+	void CacheSubPropertySearchables(TMap<UStruct*, TArray<FSearchableProperty>>& InOutCachePropertyMapping, UProperty* InProperty, const void* InValue, UStruct* InStruct)
+	{
+		if (UArrayProperty* ArrayProperty = Cast<UArrayProperty>(InProperty))
+		{
+			FScriptArrayHelper Helper(ArrayProperty, InValue);
+			for (int32 i = 0, n = Helper.Num(); i < n; ++i)
+			{
+				CacheSubPropertySearchables(InOutCachePropertyMapping, ArrayProperty->Inner, Helper.GetRawPtr(i), InStruct);
+			}
+		}
+		else if (UStructProperty* StructProperty = Cast<UStructProperty>(InProperty))
+		{
+			if (!InOutCachePropertyMapping.Find(StructProperty->Struct))
+			{
+				if (!InProperty->HasMetaData(*FFiBMD::FiBSearchableMD) || InProperty->GetBoolMetaData(*FFiBMD::FiBSearchableMD))
+				{
+					CacheSearchableProperties(InOutCachePropertyMapping, InValue, StructProperty->Struct, SEARCHABLE_FULL);
+				}
+			}
+		}
+		else if (UObjectProperty* ObjectProperty = Cast<UObjectProperty>(InProperty))
+		{
+			UObject* SubObject = ObjectProperty->GetObjectPropertyValue(InValue);
+			if (SubObject)
+			{
+				// Objects default to shallow unless they are marked as searchable
+				EGatherSearchableType SearchType = SEARCHABLE_SHALLOW;
+
+				// Check if there is any Searchable metadata
+				if (InProperty->HasMetaData(*FFiBMD::FiBSearchableMD))
+				{
+					if (!InOutCachePropertyMapping.Find(SubObject->GetClass()))
+					{
+						// Check if that metadata informs us that the property should not be searchable
+						bool bSearchable = InProperty->GetBoolMetaData(*FFiBMD::FiBSearchableMD);
+						if (bSearchable)
+						{
+							CacheSearchableProperties(InOutCachePropertyMapping, SubObject, SubObject->GetClass(), SEARCHABLE_FULL);
+						}
+					}
+				}
+			}
+		}
+	}
+
+	void CacheSearchableProperties(TMap<UStruct*, TArray<FSearchableProperty>>& InOutCachePropertyMapping, const void* InValue, UStruct* InStruct, EGatherSearchableType InSearchableType)
+	{
+		if (InValue)
+		{
+			TArray<FSearchableProperty> SearchableProperties;
+
+			for (TFieldIterator<UProperty> PropIt(InStruct); PropIt; ++PropIt)
+			{
+				UProperty* Property = *PropIt;
+				bool bIsSearchableMD = Property->GetBoolMetaData(*FFiBMD::FiBSearchableMD);
+				bool bIsShallowSearchableMD = Property->GetBoolMetaData(*FFiBMD::FiBSearchableShallowMD);
+				// It only is truly marked as not searchable if it has the metadata set to false, if the metadata is missing then we assume the searchable type that is passed in unless SEARCHABLE_AS_DESIRED
+				bool bIsMarkedNotSearchableMD = Property->HasMetaData(*FFiBMD::FiBSearchableMD) && !bIsSearchableMD;
+
+				if ((InSearchableType != SEARCHABLE_AS_DESIRED && !bIsMarkedNotSearchableMD)
+					|| bIsShallowSearchableMD || bIsSearchableMD)
+				{
+					const void* Value = Property->ContainerPtrToValuePtr<uint8>(InValue);
+
+					FSearchableProperty SearchableProperty;
+					SearchableProperty.TargetProperty = Property;
+					SearchableProperty.bIsSearchableMD = bIsSearchableMD;
+					SearchableProperty.bIsShallowSearchableMD = bIsShallowSearchableMD;
+					SearchableProperty.bIsMarkedNotSearchableMD = bIsMarkedNotSearchableMD;
+
+					if (Property->ArrayDim == 1)
+					{
+						CacheSubPropertySearchables(InOutCachePropertyMapping, Property, Value, InStruct);
+					}
+					else
+					{
+						TArray< TSharedPtr<FJsonValue> > Array;
+						for (int Index = 0; Index != Property->ArrayDim; ++Index)
+						{
+							CacheSubPropertySearchables(InOutCachePropertyMapping, Property, (char*)Value + Index * Property->ElementSize, InStruct);
+						}
+					}
+					SearchableProperties.Add(MoveTemp(SearchableProperty));
+				}
+				InOutCachePropertyMapping.Add(InStruct, SearchableProperties);
+			}
+		}
+	}
+
 	/**
 	 * Gathers all nodes from a specified graph and serializes their searchable data to Json
 	 *
@@ -823,46 +959,54 @@ namespace BlueprintSearchMetaDataHelpers
 			{
 				if(Node)
 				{
-					// Make sure we don't collect search data for nodes that are going away soon
-					if( Node->GetOuter()->IsPendingKill() )
 					{
-						continue;
-					}
+						// Make sure we don't collect search data for nodes that are going away soon
+						if (Node->GetOuter()->IsPendingKill())
+						{
+							continue;
+						}
 
-					InWriter->WriteObjectStart();
+						InWriter->WriteObjectStart();
 
-					// Retrieve the search metadata from the node, some node types may have extra metadata to be searchable.
-					TArray<struct FSearchTagDataPair> Tags;
-					Node->AddSearchMetaDataInfo(Tags);
+						// Retrieve the search metadata from the node, some node types may have extra metadata to be searchable.
+						TArray<struct FSearchTagDataPair> Tags;
+						Node->AddSearchMetaDataInfo(Tags);
 
-					// Go through the node metadata tags and put them into the Json object.
-					for(const FSearchTagDataPair& SearchData : Tags)
-					{
-						InWriter->WriteValue(SearchData.Key, SearchData.Value);
+						// Go through the node metadata tags and put them into the Json object.
+						for (const FSearchTagDataPair& SearchData : Tags)
+						{
+							InWriter->WriteValue(SearchData.Key, SearchData.Value);
+						}
 					}
 
-					// Find all the pins and extract their metadata
-					InWriter->WriteArrayStart(FFindInBlueprintSearchTags::FiB_Pins);
-					for(UEdGraphPin* Pin : Node->Pins)
 					{
-						// Hidden pins are not searchable
-						if(Pin->bHidden == false)
+						// Find all the pins and extract their metadata
+						InWriter->WriteArrayStart(FFindInBlueprintSearchTags::FiB_Pins);
+						for (UEdGraphPin* Pin : Node->Pins)
 						{
-							InWriter->WriteObjectStart();
+							// Hidden pins are not searchable
+							if (Pin->bHidden == false)
 							{
-								InWriter->WriteValue(FFindInBlueprintSearchTags::FiB_Name, Pin->GetSchema()->GetPinDisplayName(Pin));
-								InWriter->WriteValue(FFindInBlueprintSearchTags::FiB_DefaultValue, FText::FromString(Pin->GetDefaultAsString()));
+								InWriter->WriteObjectStart();
+								{
+									InWriter->WriteValue(FFindInBlueprintSearchTags::FiB_Name, Pin->GetSchema()->GetPinDisplayName(Pin));
+									InWriter->WriteValue(FFindInBlueprintSearchTags::FiB_DefaultValue, FText::FromString(Pin->GetDefaultAsString()));
+								}
+								SavePinTypeToJson(InWriter, Pin->PinType);
+								InWriter->WriteObjectEnd();
 							}
-							SavePinTypeToJson(InWriter, Pin->PinType);
-							InWriter->WriteObjectEnd();
 						}
-					}
-					InWriter->WriteArrayEnd();
+						InWriter->WriteArrayEnd();
 
-					// Only support this for nodes for now, will gather all searchable properties
-					GatherSearchableProperties(InWriter, Node, Node->GetClass());
+						if (!InWriter->CachedPropertyMapping.Find(Node->GetClass()))
+						{
+							CacheSearchableProperties(InWriter->CachedPropertyMapping, Node, Node->GetClass());
+						}
+						// Only support this for nodes for now, will gather all searchable properties
+						GatherSearchableProperties(InWriter, Node, Node->GetClass());
 
-					InWriter->WriteObjectEnd();
+						InWriter->WriteObjectEnd();
+					}
 				}
 				
 			}
@@ -887,6 +1031,11 @@ namespace BlueprintSearchMetaDataHelpers
 			{
 				for(const UEdGraph* Graph : InGraphArray)
 				{
+					// This is non-critical but should not happen and needs to be resolved
+					if (!ensure(Graph != nullptr))
+					{
+						continue;
+					}
 					InWriter->WriteObjectStart();
 
 					FGraphDisplayInfo DisplayInfo;
@@ -1488,7 +1637,7 @@ FString FFindInBlueprintSearchManager::GatherBlueprintSearchMetadata(const UBlue
 	{
 		Writer->WriteArrayStart(FFindInBlueprintSearchTags::FiB_Properties);
 		{
-			for( const FBPVariableDescription& Variable : Blueprint->NewVariables )
+			for (const FBPVariableDescription& Variable : Blueprint->NewVariables)
 			{
 				BlueprintSearchMetaDataHelpers::SaveVariableDescriptionToJson(Writer, Blueprint, Variable);
 			}
@@ -1501,6 +1650,13 @@ FString FFindInBlueprintSearchManager::GatherBlueprintSearchMetadata(const UBlue
 	BlueprintSearchMetaDataHelpers::GatherGraphSearchData(Writer, Blueprint, Blueprint->UbergraphPages, FFindInBlueprintSearchTags::FiB_UberGraphs, &SubGraphs);
 	BlueprintSearchMetaDataHelpers::GatherGraphSearchData(Writer, Blueprint, Blueprint->FunctionGraphs, FFindInBlueprintSearchTags::FiB_Functions, &SubGraphs);
 	BlueprintSearchMetaDataHelpers::GatherGraphSearchData(Writer, Blueprint, Blueprint->MacroGraphs, FFindInBlueprintSearchTags::FiB_Macros, &SubGraphs);
+
+	// Gather all interface graphs as functions
+	for (const FBPInterfaceDescription& InterfaceDesc : Blueprint->ImplementedInterfaces)
+	{
+		BlueprintSearchMetaDataHelpers::GatherGraphSearchData(Writer, Blueprint, InterfaceDesc.Graphs, FFindInBlueprintSearchTags::FiB_Functions, &SubGraphs);
+	}
+
 	// Sub graphs are processed separately so that they do not become children in the TreeView, cluttering things up if the tree is deep
 	BlueprintSearchMetaDataHelpers::GatherGraphSearchData(Writer, Blueprint, SubGraphs, FFindInBlueprintSearchTags::FiB_SubGraphs, NULL);
 
@@ -1592,6 +1748,7 @@ void FFindInBlueprintSearchManager::AddOrUpdateBlueprintSearchMetadata(UBlueprin
 	if(InBlueprint->SkeletonGeneratedClass != nullptr)
 	{
 		SearchArray[Index].Value = GatherBlueprintSearchMetadata(InBlueprint);
+		SearchArray[Index].Version = EFiBVersion::FIB_VER_LATEST;
 	}
 	SearchArray[Index].bMarkedForDeletion = false;
 }
diff --git a/Engine/Source/Editor/Kismet/Private/FindInBlueprints.cpp b/Engine/Source/Editor/Kismet/Private/FindInBlueprints.cpp
--- a/Engine/Source/Editor/Kismet/Private/FindInBlueprints.cpp
+++ b/Engine/Source/Editor/Kismet/Private/FindInBlueprints.cpp
@@ -200,6 +200,7 @@ FText FFindInBlueprintsResult::GetDisplayString() const
 
 FFindInBlueprintsGraphNode::FFindInBlueprintsGraphNode(const FText& InValue, TSharedPtr<FFindInBlueprintsResult> InParent)
 	: FFindInBlueprintsResult(InValue, InParent)
+	, GlyphBrush(nullptr)
 	, Class(nullptr)
 {
 }
diff --git a/Engine/Source/Editor/Kismet/Private/Profiler/SBlueprintProfilerView.cpp b/Engine/Source/Editor/Kismet/Private/Profiler/SBlueprintProfilerView.cpp
--- a/Engine/Source/Editor/Kismet/Private/Profiler/SBlueprintProfilerView.cpp
+++ b/Engine/Source/Editor/Kismet/Private/Profiler/SBlueprintProfilerView.cpp
@@ -27,6 +27,18 @@ void SBlueprintProfilerView::Construct(const FArguments& InArgs)
 	NumberFormat.UseGrouping = false;
 	// Update the average sample retention count.
 	FBPPerformanceData::StatSampleSize = GetDefault<UEditorExperimentalSettings>()->BlueprintProfilerAverageSampleCount;
+	// Create the display text for the user
+	if (GetDefault<UEditorExperimentalSettings>()->bBlueprintPerformanceAnalysisTools)
+	{
+		IBlueprintProfilerInterface* ProfilerInterface = FModuleManager::GetModulePtr<IBlueprintProfilerInterface>("BlueprintProfiler");
+		const bool bProfilerEnabled	= ProfilerInterface && ProfilerInterface->IsProfilerEnabled();
+		StatusText = bProfilerEnabled ? NSLOCTEXT("ActiveProfilerWidget", "NoDataText", "The Blueprint Profiler is active but does not currently have any data to display") :
+										NSLOCTEXT("InactiveProfilerWidget", "InactiveText", "The Blueprint Profiler is currently Inactive");
+	}
+	else
+	{
+		StatusText = NSLOCTEXT("DisabledProfilerWidget", "DisabledProfilerText", "The Blueprint Profiler is experimental and is currently not enabled in the editor preferences");
+	}
 	// Create the profiler view widgets
 	bTreeViewIsDirty = false;
 	CreateProfilerWidgets();
@@ -47,6 +59,8 @@ void SBlueprintProfilerView::OnGetChildrenForProfiler(FBPProfilerStatPtr InParen
 
 void SBlueprintProfilerView::OnToggleProfiler(bool bEnabled)
 {
+	StatusText = bEnabled ? NSLOCTEXT("ActiveProfilerWidget", "NoDataText", "The Blueprint Profiler is active but does not currently have any data to display") :
+							NSLOCTEXT("InactiveProfilerWidget", "InactiveText", "The Blueprint Profiler is currently Inactive");
 	if (!bEnabled)
 	{
 		// Clear profiler data, this requires work as I think the stats are currently leaking.
@@ -80,7 +94,7 @@ void SBlueprintProfilerView::CreateProfilerWidgets()
 			.VAlign(VAlign_Center)
 			[
 				SNew(STextBlock)
-				.Text(NSLOCTEXT("InactiveProfilerWidget", "InactiveText", "The Blueprint Profiler is currently Inactive"))
+				.Text(this, &SBlueprintProfilerView::GetProfilerStatusText)
 			]
 		]
 #if 0 // To do stat ranking list
diff --git a/Engine/Source/Editor/Kismet/Private/Profiler/SBlueprintProfilerView.h b/Engine/Source/Editor/Kismet/Private/Profiler/SBlueprintProfilerView.h
--- a/Engine/Source/Editor/Kismet/Private/Profiler/SBlueprintProfilerView.h
+++ b/Engine/Source/Editor/Kismet/Private/Profiler/SBlueprintProfilerView.h
@@ -60,14 +60,20 @@ protected:
 	/** Returns is the profiler widgets are currently visible */
 	EVisibility IsProfilerVisible() const { return RootTreeItems.Num() > 0 ? EVisibility::Visible : EVisibility::Collapsed; }
 
-	/** Returns is the profiler widgets are currently not visible, used to display the inactive message */
+	/** Returns if the profiler widgets are currently not visible, used to display the inactive message */
 	EVisibility IsProfilerHidden() const { return RootTreeItems.Num() == 0 ? EVisibility::Visible : EVisibility::Collapsed; }
 
+	/** Returns the profiler current status message */
+	FText GetProfilerStatusText() const { return StatusText; }
+
 	/** Called to create the child profiler widgets */
 	void CreateProfilerWidgets();
 
 protected:
 
+	/** Profiler status Text */
+	FText StatusText;
+
 	/** Number formatting options */
 	static FNumberFormattingOptions NumberFormat;
 
diff --git a/Engine/Source/Editor/Kismet/Private/SBlueprintEditorToolbar.cpp b/Engine/Source/Editor/Kismet/Private/SBlueprintEditorToolbar.cpp
--- a/Engine/Source/Editor/Kismet/Private/SBlueprintEditorToolbar.cpp
+++ b/Engine/Source/Editor/Kismet/Private/SBlueprintEditorToolbar.cpp
@@ -175,15 +175,6 @@ void FKismet2Menu::FillDebugMenu( FMenuBuilder& MenuBuilder )
 	MenuBuilder.EndSection();
 }
 
-void FKismet2Menu::FillProfilerMenu( FMenuBuilder& MenuBuilder )
-{
-	MenuBuilder.BeginSection("Profiler", LOCTEXT("ProfilerMenu_Heading", "Profiler") );
-	{
-		MenuBuilder.AddMenuEntry( FFullBlueprintEditorCommands::Get().ToggleProfiler );
-	}
-	MenuBuilder.EndSection();
-}
-
 void FKismet2Menu::SetupBlueprintEditorMenu( TSharedPtr< FExtender > Extender, FBlueprintEditor& BlueprintEditor)
 {
 	// Extend the File menu with asset actions
@@ -219,16 +210,6 @@ void FKismet2Menu::SetupBlueprintEditorMenu( TSharedPtr< FExtender > Extender, F
 					LOCTEXT("DebugMenu_ToolTip", "Open the debug menu"),
 					FNewMenuDelegate::CreateStatic( &FKismet2Menu::FillDebugMenu ),
 					"Debug");
-
-				if (GetDefault<UEditorExperimentalSettings>()->bBlueprintPerformanceAnalysisTools)
-				{
-					// Profiler
-					MenuBarBuilder.AddPullDownMenu( 
-						LOCTEXT("ProfilerMenu", "Profiler"),
-						LOCTEXT("ProfilerMenu_ToolTip", "Toggle the Profiler Menu"),
-						FNewMenuDelegate::CreateStatic( &FKismet2Menu::FillProfilerMenu ),
-						"Profiler");
-				}
 			}
 		};
 
diff --git a/Engine/Source/Editor/Kismet/Private/SBlueprintPalette.cpp b/Engine/Source/Editor/Kismet/Private/SBlueprintPalette.cpp
--- a/Engine/Source/Editor/Kismet/Private/SBlueprintPalette.cpp
+++ b/Engine/Source/Editor/Kismet/Private/SBlueprintPalette.cpp
@@ -144,9 +144,10 @@ static FString GetVarTooltip(UBlueprint* InBlueprint, UClass* VarClass, FName Va
  * @param  ActionIn		The FEdGraphSchemaAction_K2Graph action that the palette item represents.
  * @param  BlueprintIn	The blueprint currently being edited (that the action is for).
  * @param  IconOut		An icon denoting the sub-graph's type.
+ * @param  ColorOut		An output color, further denoting the specified action.
  * @param  ToolTipOut	The tooltip to display when the icon is hovered over (describing the sub-graph type).
  */
-static void GetSubGraphIcon(FEdGraphSchemaAction_K2Graph const* const ActionIn, UBlueprint const* BlueprintIn, FSlateBrush const*& IconOut, FString& ToolTipOut)
+static void GetSubGraphIcon(FEdGraphSchemaAction_K2Graph const* const ActionIn, UBlueprint const* BlueprintIn, FSlateBrush const*& IconOut, FSlateColor& ColorOut, FString& ToolTipOut)
 {
 	check(BlueprintIn != NULL);
 
@@ -224,7 +225,19 @@ static void GetSubGraphIcon(FEdGraphSchemaAction_K2Graph const* const ActionIn,
 	case EEdGraphSchemaAction_K2Graph::Interface:
 		{
 			IconOut = FEditorStyle::GetBrush(TEXT("GraphEditor.InterfaceFunction_16x"));
-			ToolTipOut = FString::Printf(*LOCTEXT("FunctionFromInterface_Tooltip", "Function (from Interface '%s')").ToString(), *ActionIn->FuncName.ToString());
+			FFormatNamedArguments Args;
+			Args.Add(TEXT("InterfaceName"), FText::FromName(ActionIn->FuncName));
+			FText TooltipText = FText::Format(LOCTEXT("FunctionFromInterface_Tooltip", "Function (from Interface '{InterfaceName}')"), Args);
+			if (UFunction* OverrideFunc = FindField<UFunction>(BlueprintIn->SkeletonGeneratedClass, ActionIn->FuncName))
+			{
+				if (UEdGraphSchema_K2::FunctionCanBePlacedAsEvent(OverrideFunc))
+				{
+					Args.Add(TEXT("BaseTooltip"), TooltipText);
+					TooltipText = FText::Format(LOCTEXT("InterfaceFunctionExpectedAsEvent_Tooltip", "{BaseTooltip}\nInterface '{InterfaceName}' is already implemented as a function graph but is expected as an event. Remove the function graph and reimplement as an event."), Args);
+					ColorOut = FLinearColor::Yellow;
+				}
+			}
+			ToolTipOut = TooltipText.ToString();
 		}
 		break;
 	case EEdGraphSchemaAction_K2Graph::Function:
@@ -312,7 +325,7 @@ static void GetPaletteItemIcon(TSharedPtr<FEdGraphSchemaAction> ActionIn, UBluep
 	else if (ActionIn->GetTypeId() == FEdGraphSchemaAction_K2Graph::StaticGetTypeId())
 	{
 		FEdGraphSchemaAction_K2Graph const* GraphAction = (FEdGraphSchemaAction_K2Graph const*)ActionIn.Get();
-		GetSubGraphIcon(GraphAction, BlueprintIn, BrushOut, ToolTipOut);
+		GetSubGraphIcon(GraphAction, BlueprintIn, BrushOut, ColorOut, ToolTipOut);
 	}
 	else if (ActionIn->GetTypeId() == FEdGraphSchemaAction_K2Delegate::StaticGetTypeId())
 	{
@@ -1446,10 +1459,12 @@ FText SBlueprintPaletteItem::GetToolTipText() const
 			FEdGraphSchemaAction_K2Graph* GraphAction = (FEdGraphSchemaAction_K2Graph*)PaletteAction.Get();
 			if (GraphAction->EdGraph != NULL)
 			{
-				FGraphDisplayInfo DisplayInfo;
-				GraphAction->EdGraph->GetSchema()->GetGraphDisplayInformation(*(GraphAction->EdGraph), DisplayInfo);
-
-				ToolTipText = FText::FromString(DisplayInfo.Tooltip);
+				if (const UEdGraphSchema* GraphSchema = GraphAction->EdGraph->GetSchema())
+				{
+					FGraphDisplayInfo DisplayInfo;
+					GraphSchema->GetGraphDisplayInformation(*(GraphAction->EdGraph), DisplayInfo);
+					ToolTipText = FText::FromString(DisplayInfo.Tooltip);
+				}
 			}
 		}
 		else if (PaletteAction->GetTypeId() == FEdGraphSchemaAction_K2Var::StaticGetTypeId())
diff --git a/Engine/Source/Editor/Kismet/Private/SMyBlueprint.cpp b/Engine/Source/Editor/Kismet/Private/SMyBlueprint.cpp
--- a/Engine/Source/Editor/Kismet/Private/SMyBlueprint.cpp
+++ b/Engine/Source/Editor/Kismet/Private/SMyBlueprint.cpp
@@ -2457,7 +2457,21 @@ bool SMyBlueprint::CanDeleteEntry() const
 
 	if (FEdGraphSchemaAction_K2Graph* GraphAction = SelectionAsGraph())
 	{
-		return GraphAction->EdGraph != NULL && GraphAction->EdGraph->bAllowDeletion;
+		if (GraphAction->EdGraph != NULL)
+		{
+			// Allow the user to delete any graphs in the interface section if the function can be placed as an event, 
+			// this allows users to resolve warnings when a previously implemented graph has been changed to be an event.
+			if (GraphAction->GetSectionID() == NodeSectionID::INTERFACE)
+			{
+				UFunction* Function = GetBlueprintObj()->SkeletonGeneratedClass->FindFunctionByName(GraphAction->EdGraph->GetFName());
+				if (UEdGraphSchema_K2::FunctionCanBePlacedAsEvent(Function))
+				{
+					return true;
+				}
+			}
+			return GraphAction->EdGraph->bAllowDeletion;
+		}
+		return false;
 	}
 	else if (FEdGraphSchemaAction_K2Delegate* DelegateAction = SelectionAsDelegate())
 	{
diff --git a/Engine/Source/Editor/Kismet/Private/SReplaceNodeReferences.cpp b/Engine/Source/Editor/Kismet/Private/SReplaceNodeReferences.cpp
--- a/Engine/Source/Editor/Kismet/Private/SReplaceNodeReferences.cpp
+++ b/Engine/Source/Editor/Kismet/Private/SReplaceNodeReferences.cpp
@@ -364,7 +364,7 @@ void SReplaceNodeReferences::GatherAllAvailableBlueprintVariables(UClass* InTarg
 				CategoryReference->Children.Add(VariableItem);
 
 				// If this is the first Child, add the category to the main Blueprint category
-				if (CategoryReference->Children.Num() == 1)
+				if (CategoryReference != BlueprintCategory && CategoryReference->Children.Num() == 1)
 				{
 					BlueprintCategory->Children.Add(CategoryReference);
 				}
diff --git a/Engine/Source/Editor/Kismet/Public/BlueprintEditor.h b/Engine/Source/Editor/Kismet/Public/BlueprintEditor.h
--- a/Engine/Source/Editor/Kismet/Public/BlueprintEditor.h
+++ b/Engine/Source/Editor/Kismet/Public/BlueprintEditor.h
@@ -137,7 +137,6 @@ public:
 	//~ Begin FAssetEditorToolkit Interface
 	virtual bool OnRequestClose() override;
 	virtual void ToolkitBroughtToFront() override;
-	virtual void SaveAsset_Execute() override;
 	// End of FAssetEditorToolkit 
 
 	//~ Begin IToolkit Interface
diff --git a/Engine/Source/Editor/Kismet/Public/FindInBlueprintManager.h b/Engine/Source/Editor/Kismet/Public/FindInBlueprintManager.h
--- a/Engine/Source/Editor/Kismet/Public/FindInBlueprintManager.h
+++ b/Engine/Source/Editor/Kismet/Public/FindInBlueprintManager.h
@@ -211,6 +211,7 @@ enum EFiBVersion
 {
 	FIB_VER_BASE = 0, // All Blueprints prior to versioning will automatically be assumed to be at 0 if they have FiB data collected
 	FIB_VER_VARIABLE_REFERENCE, // Variable references (FMemberReference) is collected in FiB
+	FIB_VER_INTERFACE_GRAPHS, // Implemented Interface Graphs is collected in FiB
 
 	// -----<new versions can be added before this line>-------------------------------------------------
 	FIB_VER_PLUS_ONE,
diff --git a/Engine/Source/Editor/Kismet/Public/SBlueprintEditorToolbar.h b/Engine/Source/Editor/Kismet/Public/SBlueprintEditorToolbar.h
--- a/Engine/Source/Editor/Kismet/Public/SBlueprintEditorToolbar.h
+++ b/Engine/Source/Editor/Kismet/Public/SBlueprintEditorToolbar.h
@@ -23,8 +23,6 @@ protected:
 
 	static void FillDebugMenu( FMenuBuilder& MenuBuilder );
 
-	static void FillProfilerMenu( FMenuBuilder& MenuBuilder );
-
 	static void FillDeveloperMenu( FMenuBuilder& MenuBuilder );
 
 private:
diff --git a/Engine/Source/Editor/KismetCompiler/Private/KismetCompiler.cpp b/Engine/Source/Editor/KismetCompiler/Private/KismetCompiler.cpp
--- a/Engine/Source/Editor/KismetCompiler/Private/KismetCompiler.cpp
+++ b/Engine/Source/Editor/KismetCompiler/Private/KismetCompiler.cpp
@@ -305,10 +305,18 @@ void FKismetCompilerContext::ValidateLink(const UEdGraphPin* PinA, const UEdGrap
 	const FPinConnectionResponse ConnectResponse = Schema->CanCreateConnection(PinA, PinB);
 
 	const bool bForbiddenConnection = (ConnectResponse.Response == CONNECT_RESPONSE_DISALLOW);
-	const bool bMissingConversion = (ConnectResponse.Response == CONNECT_RESPONSE_MAKE_WITH_CONVERSION_NODE);
+	const bool bMissingConversion   = (ConnectResponse.Response == CONNECT_RESPONSE_MAKE_WITH_CONVERSION_NODE);
 	if (bForbiddenConnection || bMissingConversion)
 	{
-		MessageLog.Warning(*FString::Printf(*LOCTEXT("PinTypeMismatch_Error", "Can't connect pins @@ and @@: %s").ToString(), *ConnectResponse.Message.ToString()), PinA, PinB);
+		const FString ErrorMessage = FString::Printf(*LOCTEXT("PinTypeMismatch_Error", "Can't connect pins @@ and @@: %s").ToString(), *ConnectResponse.Message.ToString());
+		if (ConnectResponse.IsFatal())
+		{
+			MessageLog.Error(*ErrorMessage, PinA, PinB);
+		}
+		else
+		{
+			MessageLog.Warning(*ErrorMessage, PinA, PinB);
+		}
 	}
 
 	if (PinA && PinB && PinA->Direction != PinB->Direction)
@@ -2455,8 +2463,27 @@ void FKismetCompilerContext::CreateAndProcessUbergraph()
 			const UFunction* Function = *FunctionIt;
 			const FName FunctionName = Function->GetFName();
 
+			const bool bCanImplementAsEvent = UEdGraphSchema_K2::FunctionCanBePlacedAsEvent(Function);
+			bool bExistsAsGraph = false;
+
+			// Any function that can be implemented as an event needs to check to see if there is already an interface function graph
+			// If there is, we want to warn the user that this is unexpected but proceed to successfully compile the Blueprint
+			if (bCanImplementAsEvent)
+			{
+				for (UEdGraph* InterfaceGraph : InterfaceDesc.Graphs)
+				{
+					if (InterfaceGraph->GetFName() == Function->GetFName())
+					{
+						bExistsAsGraph = true;
+
+						// Having an event override implemented as a function won't cause issues but is something the user should be aware of.
+						MessageLog.Warning(TEXT("Interface '@@' is already implemented as a function graph but is expected as an event. Remove the function graph and reimplement as an event."), InterfaceGraph);
+					}
+				}
+			}
+
 			// If this is an event, check the merged ubergraph to make sure that it has an event handler, and if not, add one
-			if (UEdGraphSchema_K2::FunctionCanBePlacedAsEvent(Function) && UEdGraphSchema_K2::CanKismetOverrideFunction(Function))
+			if (bCanImplementAsEvent && UEdGraphSchema_K2::CanKismetOverrideFunction(Function) && !bExistsAsGraph)
 			{
 				bool bFoundEntry = false;
 				// Search the cached entry points to see if we have a match
diff --git a/Engine/Source/Editor/KismetCompiler/Private/KismetCompilerMisc.cpp b/Engine/Source/Editor/KismetCompiler/Private/KismetCompilerMisc.cpp
--- a/Engine/Source/Editor/KismetCompiler/Private/KismetCompilerMisc.cpp
+++ b/Engine/Source/Editor/KismetCompiler/Private/KismetCompilerMisc.cpp
@@ -1883,12 +1883,21 @@ FBPTerminal* FKismetFunctionContext::CreateLocalTerminalFromPinAutoChooseScope(U
 {
 	check(Net);
 	bool bSharedTerm = IsEventGraph();
-
-	const bool bRegularBlueprint = Blueprint && (Blueprint->GetClass() == UBlueprint::StaticClass());
 	static FBoolConfigValueHelper UseLocalGraphVariables(TEXT("Kismet"), TEXT("bUseLocalGraphVariables"), GEngineIni);
-	const bool bUseLocalGraphVariables = UseLocalGraphVariables || bGeneratingCpp;
+
+	auto PinRepresentsDelegate = [](const FEdGraphPinType& PinType) -> bool
+	{
+		return (PinType.PinCategory == UEdGraphSchema_K2::PC_Delegate) 
+			|| (PinType.PinCategory == UEdGraphSchema_K2::PC_MCDelegate);
+	};
+	// Local graph Variables cannot be used with persistent UberGraph frame
+	// Current implementation can change BP behavior, so it cannot be used to generate C++
+	const bool bUseLocalGraphVariables = UseLocalGraphVariables
+		//This is a workaround for UE-24210. UHT doesn't support  delegates with native type form: OtherOwnerClass::DelegateType.
+		|| (bGeneratingCpp && PinRepresentsDelegate(Net->PinType)); 
+
 	const bool OutputPin = EEdGraphPinDirection::EGPD_Output == Net->Direction;
-	if (bSharedTerm && bRegularBlueprint && bUseLocalGraphVariables && OutputPin)
+	if (bSharedTerm && bUseLocalGraphVariables && OutputPin)
 	{
 		BP_SCOPED_COMPILER_EVENT_STAT(EKismetCompilerStats_ChooseTerminalScope);
 
diff --git a/Engine/Source/Editor/UnrealEd/Classes/Commandlets/NativeCodeGenCommandlet.h b/Engine/Source/Editor/UnrealEd/Classes/Commandlets/NativeCodeGenCommandlet.h
new file mode 100644
--- /dev/null
+++ b/Engine/Source/Editor/UnrealEd/Classes/Commandlets/NativeCodeGenCommandlet.h
@@ -0,0 +1,18 @@
+// Copyright 1998-2015 Epic Games, Inc. All Rights Reserved.
+
+#pragma once
+
+#include "Commandlets/Commandlet.h"
+#include "NativeCodeGenCommandlet.generated.h"
+
+UCLASS()
+class UNativeCodeGenCommandlet
+	: public UCommandlet
+{
+	GENERATED_UCLASS_BODY()
+
+	//~ Begin UCommandlet Interface
+	virtual int32 Main(const FString& Params) override;
+	//~ End UCommandlet Interface
+};
+
diff --git a/Engine/Source/Editor/UnrealEd/Classes/CookOnTheSide/CookOnTheFlyServer.h b/Engine/Source/Editor/UnrealEd/Classes/CookOnTheSide/CookOnTheFlyServer.h
--- a/Engine/Source/Editor/UnrealEd/Classes/CookOnTheSide/CookOnTheFlyServer.h
+++ b/Engine/Source/Editor/UnrealEd/Classes/CookOnTheSide/CookOnTheFlyServer.h
@@ -673,7 +673,6 @@ public:
 		FProcHandle ProcessHandle;
 		FString ResponseFileName;
 		FString BaseResponseFileName;
-		FString ManifestFilename;
 		void* ReadPipe;
 		int32 ReturnCode;
 		bool bFinished;
@@ -693,7 +692,8 @@ private:
 			bForceEnableCompressedPackages(false),
 			bForceDisableCompressedPackages(false),
 			bIsChildCooker(false),
-			bDisableUnsolicitedPackages(false)
+			bDisableUnsolicitedPackages(false),
+			ChildCookIdentifier(-1)
 		{ }
 
 		/** Should we test for UObject leaks */
@@ -730,8 +730,8 @@ private:
 		bool bForceDisableCompressedPackages;
 		bool bIsChildCooker;
 		bool bDisableUnsolicitedPackages;
+		int32 ChildCookIdentifier;
 		FString ChildCookFilename;
-		FString ChildManifestFilename;
 		TSet<FName> ChildUnsolicitedPackages;
 		TArray<FChildCooker> ChildCookers;
 		TArray<FName> TargetPlatformNames;
@@ -890,7 +890,7 @@ public:
 		FString CreateReleaseVersion;
 		FString BasedOnReleaseVersion;
 		FString ChildCookFileName; // if we are the child cooker 
-		FString ChildManifestFilename; // again, only if you are the child cooker
+		int32 ChildCookIdentifier; // again, only if you are the child cooker
 		bool bGenerateStreamingInstallManifests; 
 		bool bGenerateDependenciesForMaps; 
 		bool bErrorOnEngineContentUse; // this is a flag for dlc, will cause the cooker to error if the dlc references engine content
@@ -898,6 +898,7 @@ public:
 		FCookByTheBookStartupOptions() :
 			CookOptions(ECookByTheBookOptions::None),
 			DLCName(FString()),
+			ChildCookIdentifier(-1),
 			bGenerateStreamingInstallManifests(false),
 			bGenerateDependenciesForMaps(false),
 			bErrorOnEngineContentUse(false),
diff --git a/Engine/Source/Editor/UnrealEd/Private/Commandlets/CookCommandlet.cpp b/Engine/Source/Editor/UnrealEd/Private/Commandlets/CookCommandlet.cpp
--- a/Engine/Source/Editor/UnrealEd/Private/Commandlets/CookCommandlet.cpp
+++ b/Engine/Source/Editor/UnrealEd/Private/Commandlets/CookCommandlet.cpp
@@ -463,10 +463,7 @@ bool UCookCommandlet::SaveCookedPackage( UPackage* Package, uint32 SaveFlags, bo
 				else
 				{
 					ESavePackageResult Result = GEditor->Save(Package, World, Flags, *PlatFilename, GError, NULL, bSwap, false, SaveFlags, Target, FDateTime::MinValue());
-					if (Result == ESavePackageResult::ReplaceCompletely || Result == ESavePackageResult::GenerateStub)
-					{
-						IBlueprintNativeCodeGenModule::Get().Convert(Package, Result == ESavePackageResult::ReplaceCompletely ? EReplacementResult::ReplaceCompletely : EReplacementResult::GenerateStub);
-					}
+					IBlueprintNativeCodeGenModule::Get().Convert(Package, Result, *(Target->PlatformName()));
 					bSavedCorrectly &= (Result == ESavePackageResult::ReplaceCompletely || Result == ESavePackageResult::GenerateStub || Result == ESavePackageResult::Success);
 				}
 				
@@ -1064,8 +1061,8 @@ bool UCookCommandlet::NewCook( const TArray<ITargetPlatform*>& Platforms, TArray
 	FString ChildCookFile;
 	FParse::Value(*Params, TEXT("cookchild="), ChildCookFile);
 
-	FString ChildManifestFilename;
-	FParse::Value(*Params, TEXT("childmanifest="), ChildManifestFilename);
+	int32 ChildCookIdentifier = -1;
+	FParse::Value(*Params, TEXT("childIdentifier="), ChildCookIdentifier);
 
 	int32 NumProcesses = 0;
 	FParse::Value(*Params, TEXT("numcookerstospawn="), NumProcesses);
@@ -1170,17 +1167,6 @@ bool UCookCommandlet::NewCook( const TArray<ITargetPlatform*>& Platforms, TArray
 	AlwaysCookMapList.Append(MapList);
 	Swap(MapList, AlwaysCookMapList);
 
-	FCookCommandParams CookParams(FCommandLine::Get());
-	if (CookParams.bRunConversion)
-	{
-		const UBlueprintNativeCodeGenConfig* ConfigSettings = GetDefault<UBlueprintNativeCodeGenConfig>();
-		TMap<UObject*, UClass*> ClassReplacementMap;
-		ClassReplacementMap.Add(UUserDefinedEnum::StaticClass(), UEnum::StaticClass());
-		ClassReplacementMap.Add(UUserDefinedStruct::StaticClass(), UScriptStruct::StaticClass());
-		ClassReplacementMap.Add(UBlueprintGeneratedClass::StaticClass(), UDynamicClass::StaticClass());
-		FScriptCookReplacementCoordinator::Create(CookParams.bRunConversion, ConfigSettings->ExcludedAssetTypes, ConfigSettings->ExcludedBlueprintTypes, ClassReplacementMap);
-	}
-
 	//////////////////////////////////////////////////////////////////////////
 	// start cook by the book 
 	ECookByTheBookOptions CookOptions = ECookByTheBookOptions::None;
@@ -1208,7 +1194,7 @@ bool UCookCommandlet::NewCook( const TArray<ITargetPlatform*>& Platforms, TArray
 	StartupOptions.bGenerateDependenciesForMaps = Switches.Contains(TEXT("GenerateDependenciesForMaps"));
 	StartupOptions.bGenerateStreamingInstallManifests = bGenerateStreamingInstallManifests;
 	StartupOptions.ChildCookFileName = ChildCookFile;
-	StartupOptions.ChildManifestFilename = ChildManifestFilename;
+	StartupOptions.ChildCookIdentifier = ChildCookIdentifier;
 	StartupOptions.NumProcesses = NumProcesses;
 
 	CookOnTheFlyServer->StartCookByTheBook( StartupOptions );
diff --git a/Engine/Source/Editor/UnrealEd/Private/Commandlets/NativeCodeGenCommandlet.cpp b/Engine/Source/Editor/UnrealEd/Private/Commandlets/NativeCodeGenCommandlet.cpp
new file mode 100644
--- /dev/null
+++ b/Engine/Source/Editor/UnrealEd/Private/Commandlets/NativeCodeGenCommandlet.cpp
@@ -0,0 +1,38 @@
+// Copyright 1998-2015 Epic Games, Inc. All Rights Reserved.
+
+#include "UnrealEd.h"
+
+#include "BlueprintNativeCodeGenModule.h"
+#include "Commandlets/NativeCodeGenCommandlet.h"
+
+DEFINE_LOG_CATEGORY_STATIC(LogNativeCodeGenCommandletCommandlet, Log, All);
+
+UNativeCodeGenCommandlet::UNativeCodeGenCommandlet(const FObjectInitializer& ObjectInitializer)
+	: Super(ObjectInitializer)
+{
+}
+
+int32 UNativeCodeGenCommandlet::Main(const FString& Params)
+{
+	TArray<FString> Platforms, Switches;
+	ParseCommandLine(*Params, Platforms, Switches);
+
+	if (Platforms.Num() == 0)
+	{
+		UE_LOG(LogNativeCodeGenCommandletCommandlet, Warning, TEXT("Missing platforms argument, should be -run=NativeCodeGen platform1 platform2, eg -run=NativeCodeGen windowsnoeditor"));
+		return 0;
+	}
+
+	TArray< TPair< FString, FString > > CodeGenTargets;
+	{
+		for (auto& Platform : Platforms)
+		{
+			// If you change this target path you must also update logic in CookCommand.Automation.CS. Passing a single directory around is cumbersome for testing, so I have hard coded it.
+			CodeGenTargets.Push(TPairInitializer<FString, FString>(Platform, FString(FPaths::Combine(*FPaths::GameIntermediateDir(), *Platform))));
+		}
+	}
+
+	IBlueprintNativeCodeGenModule::InitializeModuleForRerunDebugOnly(CodeGenTargets);
+
+	return 0;
+}
diff --git a/Engine/Source/Editor/UnrealEd/Private/CookOnTheFlyServer.cpp b/Engine/Source/Editor/UnrealEd/Private/CookOnTheFlyServer.cpp
--- a/Engine/Source/Editor/UnrealEd/Private/CookOnTheFlyServer.cpp
+++ b/Engine/Source/Editor/UnrealEd/Private/CookOnTheFlyServer.cpp
@@ -2779,10 +2779,7 @@ ESavePackageResult UCookOnTheFlyServer::SaveCookedPackage(UPackage* Package, uin
 				{
 					SCOPE_TIMER(GEditorSavePackage);
 					Result = GEditor->Save(Package, World, Flags, *PlatFilename, GError, NULL, bSwap, false, SaveFlags, Target, FDateTime::MinValue(), false);
-					if (Result == ESavePackageResult::ReplaceCompletely || Result == ESavePackageResult::GenerateStub)
-					{
-						IBlueprintNativeCodeGenModule::Get().Convert(Package, Result == ESavePackageResult::ReplaceCompletely ? EReplacementResult::ReplaceCompletely : EReplacementResult::GenerateStub);
-					}
+					IBlueprintNativeCodeGenModule::Get().Convert(Package, Result, *(Target->PlatformName()));
 					INC_INT_STAT(SavedPackage, 1);
 				}
 
@@ -4095,20 +4092,26 @@ void UCookOnTheFlyServer::CookByTheBookFinished()
 			UncookedPackageList.Append(UncookedPackage.ToString() + TEXT("\n\r"));
 		}
 		FFileHelper::SaveStringToFile(UncookedPackageList, *GetChildCookerResultFilename(CookByTheBookOptions->ChildCookFilename));
-		IBlueprintNativeCodeGenModule::Get().SaveManifest(*GetChildCookerManifestFilename(CookByTheBookOptions->ChildManifestFilename));
+		if (IBlueprintNativeCodeGenModule::IsNativeCodeGenModuleLoaded())
+		{
+			IBlueprintNativeCodeGenModule::Get().SaveManifest(CookByTheBookOptions->ChildCookIdentifier);
+		}
 	}
 	else
 	{
 		CleanUpChildCookers();
 
-		IBlueprintNativeCodeGenModule& CodeGenModule = IBlueprintNativeCodeGenModule::Get();
-		// merge the manifest for the bluepritn code generator:
-		for (auto& Cooker : CookByTheBookOptions->ChildCookers)
+		if (IBlueprintNativeCodeGenModule::IsNativeCodeGenModuleLoaded())
 		{
-			CodeGenModule.MergeManifest(*GetChildCookerManifestFilename(Cooker.ManifestFilename));
+			IBlueprintNativeCodeGenModule& CodeGenModule = IBlueprintNativeCodeGenModule::Get();
+			// merge the manifest for the blueprint code generator:
+			for (int32 I = 0; I < CookByTheBookOptions->ChildCookers.Num(); ++I )
+			{
+				CodeGenModule.MergeManifest(I);
+			}
+
+			CodeGenModule.FinalizeManifest();
 		}
-		
-		CodeGenModule.FinalizeManifest();
 
 		check(CookByTheBookOptions->ChildUnsolicitedPackages.Num() == 0);
 		SCOPE_COOKING_STAT(SavingAssetRegistry);
@@ -4403,7 +4406,6 @@ void UCookOnTheFlyServer::ValidateCookByTheBookSettings() const
 
 void UCookOnTheFlyServer::StartCookByTheBook( const FCookByTheBookStartupOptions& CookByTheBookStartupOptions )
 {
-
 	const TArray<ITargetPlatform*>& TargetPlatforms = CookByTheBookStartupOptions.TargetPlatforms;
 	const TArray<FString>& CookMaps = CookByTheBookStartupOptions.CookMaps;
 	const TArray<FString>& CookDirectories = CookByTheBookStartupOptions.CookDirectories;
@@ -4430,7 +4432,7 @@ void UCookOnTheFlyServer::StartCookByTheBook( const FCookByTheBookStartupOptions
 	CookByTheBookOptions->bIsChildCooker = CookByTheBookStartupOptions.ChildCookFileName.Len() > 0 ? true : false;
 	CookByTheBookOptions->ChildCookFilename = CookByTheBookStartupOptions.ChildCookFileName;
 	CookByTheBookOptions->bDisableUnsolicitedPackages = !!(CookOptions & ECookByTheBookOptions::DisableUnsolicitedPackages);
-	CookByTheBookOptions->ChildManifestFilename = CookByTheBookStartupOptions.ChildManifestFilename;
+	CookByTheBookOptions->ChildCookIdentifier = CookByTheBookStartupOptions.ChildCookIdentifier;
 
 	NeverCookPackageList.Empty();
 	{
@@ -4499,6 +4501,18 @@ void UCookOnTheFlyServer::StartCookByTheBook( const FCookByTheBookStartupOptions
 
 	InitializeSandbox();
 
+	bool bRunConversion = FParse::Param(FCommandLine::Get(), TEXT("NativizeAssets"));
+	if (bRunConversion)
+	{
+		FNativeCodeGenInitData CodeGenData;
+		for (auto Entry : CookByTheBookStartupOptions.TargetPlatforms)
+		{
+			// If you change this target path you must also update logic in CookCommand.Automation.CS. Passing a single directory around is cumbersome for testing, so I have hard coded it.
+			CodeGenData.CodegenTargets.Push(TPairInitializer<FString, FString>(Entry->PlatformName(), FString(FPaths::Combine( *FPaths::GameIntermediateDir(), *(Entry->PlatformName())))));
+		}
+		CodeGenData.ManifestIdentifier = CookByTheBookStartupOptions.ChildCookIdentifier;
+		IBlueprintNativeCodeGenModule::InitializeModule(CodeGenData);
+	}
 
 	// need to test this out
 	/*if (RecompileChangedShaders(TargetPlatformNames))
@@ -4896,7 +4910,7 @@ void UCookOnTheFlyServer::StartChildCookers(int32 NumCookersToSpawn, const TArra
 
 		ChildCooker.ResponseFileName = FPaths::CreateTempFilename(*(FPaths::GameSavedDir() / TEXT("CookingTemp")));
 		ChildCooker.BaseResponseFileName = FPaths::GetBaseFilename(ChildCooker.ResponseFileName);
-		ChildCooker.ManifestFilename = FString(*(FPaths::GameSavedDir() / TEXT("CookingTemp") / *FString::Printf(TEXT("BPCodeManifest_%d.json"), CookerCounter)));
+
 		// FArchive* ResponseFile = IFileManager::CreateFileWriter(ChildCooker.UniqueTempName);
 		FString ResponseFileText;
 
@@ -4932,8 +4946,8 @@ void UCookOnTheFlyServer::StartChildCookers(int32 NumCookersToSpawn, const TArra
 
 		// default commands
 		// multiprocess tells unreal in general we shouldn't do things like save ddc, clean shader working directory, and other various multiprocess unsafe things
-		FString CommandLine = FString::Printf(TEXT("\"%s\" -run=cook -multiprocess -targetplatform=%s -cookchild=\"%s\" -abslog=\"%sLog.txt\" -childmanifest=\"%s\" %s"), 
-			*FPaths::GetProjectFilePath(), *TargetPlatformString, *ChildCooker.ResponseFileName, *ChildCooker.ResponseFileName, *ChildCooker.ManifestFilename, *ExtraCmdParams);
+		FString CommandLine = FString::Printf(TEXT("\"%s\" -run=cook -multiprocess -targetplatform=%s -cookchild=\"%s\" -abslog=\"%sLog.txt\" -childIdentifier=%d %s"), 
+			*FPaths::GetProjectFilePath(), *TargetPlatformString, *ChildCooker.ResponseFileName, *ChildCooker.ResponseFileName, CookerCounter, *ExtraCmdParams);
 
 		auto KeepCommandlineValue = [&](const TCHAR* CommandlineToKeep)
 		{
@@ -4957,7 +4971,7 @@ void UCookOnTheFlyServer::StartChildCookers(int32 NumCookersToSpawn, const TArra
 		};
 
 
-
+		KeepCommandlineParam(TEXT("NativizeAssets"));
 		KeepCommandlineValue(TEXT("ddc="));
 		KeepCommandlineParam(TEXT("SkipEditorContent"));
 		KeepCommandlineParam(TEXT("compressed"));
diff --git a/Engine/Source/Editor/UnrealEd/Private/Editor.cpp b/Engine/Source/Editor/UnrealEd/Private/Editor.cpp
--- a/Engine/Source/Editor/UnrealEd/Private/Editor.cpp
+++ b/Engine/Source/Editor/UnrealEd/Private/Editor.cpp
@@ -892,8 +892,16 @@ namespace EditorUtilities
 			UProperty* InnerProperty = ArrayProperty->Inner;
 			int32 Num = SourceArrayHelper.Num();
 
-			TargetArrayHelper.EmptyAndAddUninitializedValues(Num);
-
+			// here we emulate UArrayProperty::CopyValuesInternal()
+			if (!(InnerProperty->PropertyFlags & CPF_IsPlainOldData))
+			{
+				TargetArrayHelper.EmptyAndAddValues(Num);
+			}
+			else
+			{
+				TargetArrayHelper.EmptyAndAddUninitializedValues(Num);
+			}
+			
 			for (int32 Index = 0; Index < Num; Index++)
 			{
 				CopySinglePropertyRecursive(SourceArrayHelper.GetRawPtr(Index), TargetArrayHelper.GetRawPtr(Index), InTargetObject, InnerProperty);
diff --git a/Engine/Source/Editor/UnrealEd/Private/Kismet2/BlueprintEditorUtils.cpp b/Engine/Source/Editor/UnrealEd/Private/Kismet2/BlueprintEditorUtils.cpp
--- a/Engine/Source/Editor/UnrealEd/Private/Kismet2/BlueprintEditorUtils.cpp
+++ b/Engine/Source/Editor/UnrealEd/Private/Kismet2/BlueprintEditorUtils.cpp
@@ -1910,6 +1910,11 @@ void FBlueprintEditorUtils::PostDuplicateBlueprint(UBlueprint* Blueprint, bool b
 			UObject* OldCDO = OldBPGC->GetDefaultObject();
 			check(OldCDO != NULL);
 
+			if (FBlueprintDuplicationScopeFlags::HasAnyFlag(FBlueprintDuplicationScopeFlags::ValidatePinsUsingSourceClass))
+			{
+				Blueprint->OriginalClass = OldBPGC;
+			}
+
 			// Grab the old class templates, which needs to be moved to the new class
 			USimpleConstructionScript* SCSRootNode = Blueprint->SimpleConstructionScript;
 			Blueprint->SimpleConstructionScript = NULL;
@@ -5466,51 +5471,14 @@ void FBlueprintEditorUtils::RemoveInterface(UBlueprint* Blueprint, const FName&
 	{
 		FBPInterfaceDescription& CurrentInterface = Blueprint->ImplementedInterfaces[Idx];
 
-		const UEdGraphSchema_K2* Schema = GetDefault<UEdGraphSchema_K2>();
-
 		// Remove all the graphs that we implemented
 		for(TArray<UEdGraph*>::TIterator it(CurrentInterface.Graphs); it; ++it)
 		{
 			UEdGraph* CurrentGraph = *it;
 			if(bPreserveFunctions)
 			{
-				CurrentGraph->bAllowDeletion = true;
-				CurrentGraph->bAllowRenaming = true;
-				CurrentGraph->bEditable = true;
-				CurrentGraph->InterfaceGuid.Invalidate();
-
-				// We need to flag the entry node to make sure that the compiled function is callable
-				Schema->AddExtraFunctionFlags(CurrentGraph, (FUNC_BlueprintCallable|FUNC_BlueprintEvent|FUNC_Public));
-				Schema->MarkFunctionEntryAsEditable(CurrentGraph, true);
-
+				PromoteGraphFromInterfaceOverride(Blueprint, CurrentGraph);
 				Blueprint->FunctionGraphs.Add(CurrentGraph);
-
-				// Move all non-exec pins from the function entry node to being user defined pins
-				TArray<UK2Node_FunctionEntry*> FunctionEntryNodes;
-				CurrentGraph->GetNodesOfClass(FunctionEntryNodes);
-				if (FunctionEntryNodes.Num() > 0)
-				{
-					UK2Node_FunctionEntry* FunctionEntry = FunctionEntryNodes[0];
-					FunctionEntry->PromoteFromInterfaceOverride();
-				}
-
-				// Move all non-exec pins from the function result node to being user defined pins
-				TArray<UK2Node_FunctionResult*> FunctionResultNodes;
-				CurrentGraph->GetNodesOfClass(FunctionResultNodes);
-				if (FunctionResultNodes.Num() > 0)
-				{
-					UK2Node_FunctionResult* PrimaryFunctionResult = FunctionResultNodes[0];
-					PrimaryFunctionResult->PromoteFromInterfaceOverride();
-
-					// Reconstruct all result nodes so they update their pins accordingly
-					for (UK2Node_FunctionResult* FunctionResult : FunctionResultNodes)
-					{
-						if (PrimaryFunctionResult != FunctionResult)
-						{
-							FunctionResult->PromoteFromInterfaceOverride(false);
-						}
-					}
-				}
 			}
 			else
 			{
@@ -5542,6 +5510,7 @@ void FBlueprintEditorUtils::RemoveInterface(UBlueprint* Blueprint, const FName&
 					{
 						UEdGraphPin* CurrentPin = *PinIt;
 						UEdGraphPin* TargetPin = NewEvent->FindPinChecked(CurrentPin->PinName);
+						const UEdGraphSchema_K2* Schema = GetDefault<UEdGraphSchema_K2>();
 						Schema->MovePinLinks(*CurrentPin, *TargetPin);
 					}
 				}
@@ -5558,6 +5527,46 @@ void FBlueprintEditorUtils::RemoveInterface(UBlueprint* Blueprint, const FName&
 	}
 }
 
+void FBlueprintEditorUtils::PromoteGraphFromInterfaceOverride(UBlueprint* InBlueprint, UEdGraph* InInterfaceGraph)
+{
+	InInterfaceGraph->bAllowDeletion = true;
+	InInterfaceGraph->bAllowRenaming = true;
+	InInterfaceGraph->bEditable = true;
+	InInterfaceGraph->InterfaceGuid.Invalidate();
+
+	// We need to flag the entry node to make sure that the compiled function is callable
+	const UEdGraphSchema_K2* Schema = GetDefault<UEdGraphSchema_K2>();
+	Schema->AddExtraFunctionFlags(InInterfaceGraph, (FUNC_BlueprintCallable | FUNC_BlueprintEvent | FUNC_Public));
+	Schema->MarkFunctionEntryAsEditable(InInterfaceGraph, true);
+
+	// Move all non-exec pins from the function entry node to being user defined pins
+	TArray<UK2Node_FunctionEntry*> FunctionEntryNodes;
+	InInterfaceGraph->GetNodesOfClass(FunctionEntryNodes);
+	if (FunctionEntryNodes.Num() > 0)
+	{
+		UK2Node_FunctionEntry* FunctionEntry = FunctionEntryNodes[0];
+		FunctionEntry->PromoteFromInterfaceOverride();
+	}
+
+	// Move all non-exec pins from the function result node to being user defined pins
+	TArray<UK2Node_FunctionResult*> FunctionResultNodes;
+	InInterfaceGraph->GetNodesOfClass(FunctionResultNodes);
+	if (FunctionResultNodes.Num() > 0)
+	{
+		UK2Node_FunctionResult* PrimaryFunctionResult = FunctionResultNodes[0];
+		PrimaryFunctionResult->PromoteFromInterfaceOverride();
+
+		// Reconstruct all result nodes so they update their pins accordingly
+		for (UK2Node_FunctionResult* FunctionResult : FunctionResultNodes)
+		{
+			if (PrimaryFunctionResult != FunctionResult)
+			{
+				FunctionResult->PromoteFromInterfaceOverride(false);
+			}
+		}
+	}
+}
+
 void FBlueprintEditorUtils::CleanNullGraphReferencesRecursive(UEdGraph* Graph)
 {
 	for (int32 GraphIndex = 0; GraphIndex < Graph->SubGraphs.Num(); )
@@ -5703,29 +5712,101 @@ void FBlueprintEditorUtils::ConformCallsToParentFunctions(UBlueprint* Blueprint)
 	}
 }
 
-// Makes sure that all events we handle exist, and replace with custom events if not
-void FBlueprintEditorUtils::ConformImplementedEvents(UBlueprint* Blueprint)
+namespace
 {
-	check(NULL != Blueprint);
+	static bool ExtendedIsParent(const UClass* Parent, const UClass* Child)
+	{
+		if (Parent && Child)
+		{
+			if (Child->IsChildOf(Parent))
+			{
+				return true;
+			}
+
+			if (Parent->ClassGeneratedBy)
+			{
+				if (Parent->ClassGeneratedBy == Child->ClassGeneratedBy)
+				{
+					return true;
+				}
+
+				if (const UBlueprint* ParentBP = Cast<UBlueprint>(Parent->ClassGeneratedBy))
+				{
+					if (ParentBP->SkeletonGeneratedClass && Child->IsChildOf(ParentBP->SkeletonGeneratedClass))
+					{
+						return true;
+					}
 
-	// Collect all event graph names
-	// @TODO: currently, in the compile process, Blueprint->EventGraphs gets 
-	//        filled out (in FKismetCompilerContext::CreateFunctionStubForEvent)
-	//        after this function is called... normally this would cause a 
-	//        problem here (as we're relying on a stale set from the last time  
-	//        this Blueprint was compiled), but because both the skeleton class 
-	//        and generated class execute this and CreateFunctionStubForEvent(), 
-	//        then the second time through (for the generated class) EventGraphs 
-	//        should be accurate
-	TArray<FName> EventGraphNames;
-	for(int EventGraphIndex = 0; EventGraphIndex < Blueprint->EventGraphs.Num(); ++EventGraphIndex)
+					if (ParentBP->GeneratedClass && Child->IsChildOf(ParentBP->GeneratedClass))
+					{
+						return true;
+					}
+				}
+			}
+		}
+		return false;
+	}
+
+	static void FixOverriddenEventSignature(UK2Node_Event* EventNode, UBlueprint* Blueprint, UEdGraph* CurrentGraph)
 	{
-		UEdGraph* EventGraph = Blueprint->EventGraphs[EventGraphIndex];
-		if(EventGraph)
+		check(EventNode && Blueprint && CurrentGraph);
+		UClass* CurrentClass = EventNode->GetBlueprintClassFromNode();
+		FMemberReference& FuncRef = EventNode->EventReference;
+		const FName EventFuncName = FuncRef.GetMemberName();
+		ensure(EventFuncName != NAME_None);
+		ensure(!EventNode->IsA<UK2Node_CustomEvent>());
+
+		const UFunction* TargetFunction = FuncRef.ResolveMember<UFunction>(CurrentClass);
+		const UClass* FuncOwnerClass = FuncRef.GetMemberParentClass(CurrentClass);
+		const bool bFunctionOwnerIsNotParentOfClass = !FuncOwnerClass || !ExtendedIsParent(FuncOwnerClass, CurrentClass);
+		const bool bNeedsToBeFixed = !TargetFunction || bFunctionOwnerIsNotParentOfClass;
+		if (bNeedsToBeFixed)
 		{
-			EventGraphNames.AddUnique(EventGraph->GetFName());
+			const UClass* SuperClass = CurrentClass->GetSuperClass();
+			const UFunction* ActualTargetFunction = SuperClass ? SuperClass->FindFunctionByName(EventFuncName) : nullptr;
+			if (ActualTargetFunction)
+			{
+				ensure(TargetFunction != ActualTargetFunction);
+				if (!ensure(!TargetFunction || TargetFunction->IsSignatureCompatibleWith(ActualTargetFunction)))
+				{
+					UE_LOG(LogBlueprint, Error
+						, TEXT("FixOverriddenEventSignature function \"%s\" is not compatible with \"%s\" node \"%s\"")
+						, *GetPathNameSafe(ActualTargetFunction), *GetPathNameSafe(TargetFunction), *GetPathNameSafe(EventNode));
+				}
+
+				ensure(GetDefault<UEdGraphSchema_K2>()->FunctionCanBePlacedAsEvent(ActualTargetFunction));
+				FuncRef.SetFromField<UFunction>(ActualTargetFunction, false);
+
+				// Emit something to the log to indicate that we've made a change
+				FFormatNamedArguments Args;
+				Args.Add(TEXT("NodeTitle"), EventNode->GetNodeTitle(ENodeTitleType::ListView));
+				Args.Add(TEXT("EventNodeName"), FText::FromString(EventNode->GetName()));
+				Blueprint->Message_Note(FText::Format(LOCTEXT("EventSignatureFixed_Note", "{NodeTitle} ({EventNodeName}) had an invalid function signature - it has now been fixed."), Args).ToString());
+			}
+			else
+			{
+				UEdGraphNode* CustomEventNode = CurrentGraph->GetSchema()->CreateSubstituteNode(EventNode, CurrentGraph, NULL);
+				if (ensure(CustomEventNode))
+				{
+					// Destroy the old event node (this will also break all pin links and remove it from the graph)
+					EventNode->DestroyNode();
+					// Add the new custom event node to the graph
+					CurrentGraph->Nodes.Add(CustomEventNode);
+					// Emit something to the log to indicate that we've made a change
+					FFormatNamedArguments Args;
+					Args.Add(TEXT("NodeTitle"), EventNode->GetNodeTitle(ENodeTitleType::ListView));
+					Args.Add(TEXT("EventNodeName"), FText::FromString(EventNode->GetName()));
+					Blueprint->Message_Note(FText::Format(LOCTEXT("EventNodeReplaced_Note", "{NodeTitle} ({EventNodeName}) was not valid for this Blueprint - it has been converted to a custom event."), Args).ToString());
+				}
+			}
 		}
 	}
+}
+
+// Makes sure that all events we handle exist, and replace with custom events if not
+void FBlueprintEditorUtils::ConformImplementedEvents(UBlueprint* Blueprint)
+{
+	check(NULL != Blueprint);
 
 	// Collect all implemented interface classes
 	TArray<UClass*> ImplementedInterfaceClasses;
@@ -5756,55 +5837,7 @@ void FBlueprintEditorUtils::ConformImplementedEvents(UBlueprint* Blueprint)
 					const bool bEventNodeUsedByInterface = ImplementedInterfaceClasses.Contains(EventNode->EventReference.GetMemberParentClass(EventNode->GetBlueprintClassFromNode()));
 					if (Blueprint->GeneratedClass && !bEventNodeUsedByInterface)
 					{
-						FMemberReference& FuncRef = EventNode->EventReference;
-						const FName EventFuncName = FuncRef.GetMemberName();
-
-						// See if the generated class implements an event with the given function signature
-						const UFunction* TargetFunction = EventNode->EventReference.ResolveMember<UFunction>(Blueprint->GeneratedClass);
-						if (TargetFunction || EventGraphNames.Contains(EventFuncName))
-						{
-							UClass* FuncOwnerClass = FuncRef.GetMemberParentClass(EventNode->GetBlueprintClassFromNode());
-							// The generated class implements the event but the function signature is not up-to-date
-							if (!Blueprint->GeneratedClass->IsChildOf(FuncOwnerClass))
-							{
-								FFormatNamedArguments Args;
-								Args.Add(TEXT("NodeTitle"), EventNode->GetNodeTitle(ENodeTitleType::ListView));
-								Args.Add(TEXT("EventNodeName"), FText::FromString(EventNode->GetName()));
-
-								// Emit something to the log to indicate that we've made a change
-								Blueprint->Message_Note( FText::Format(LOCTEXT("EventSignatureFixed_Note", "{NodeTitle} ({EventNodeName}) had an invalid function signature - it has now been fixed."), Args).ToString() );
-
-								// Fix up the event signature
-								if (TargetFunction != nullptr)
-								{
-								EventNode->EventReference.SetFromField<UFunction>(TargetFunction, false);
-							}
-								else
-								{
-									EventNode->EventReference.SetExternalMember(EventFuncName, Blueprint->GeneratedClass);
-								}
-							}
-						}
-						else
-						{
-							// The generated class does not implement this event, so replace it with a custom event node instead (this will persist any existing connections)
-							UEdGraphNode* CustomEventNode = CurrentGraph->GetSchema()->CreateSubstituteNode(EventNode, CurrentGraph, NULL);
-							if(CustomEventNode)
-							{
-								FFormatNamedArguments Args;
-								Args.Add(TEXT("NodeTitle"), EventNode->GetNodeTitle(ENodeTitleType::ListView));
-								Args.Add(TEXT("EventNodeName"), FText::FromString(EventNode->GetName()));
-
-								// Emit something to the log to indicate that we've made a change
-								Blueprint->Message_Note(FText::Format(LOCTEXT("EventNodeReplaced_Note", "{NodeTitle} ({EventNodeName}) was not valid for this Blueprint - it has been converted to a custom event."), Args).ToString());
-
-								// Destroy the old event node (this will also break all pin links and remove it from the graph)
-								EventNode->DestroyNode();
-
-								// Add the new custom event node to the graph
-								CurrentGraph->Nodes.Add(CustomEventNode);
-							}
-						}
+						FixOverriddenEventSignature(EventNode, Blueprint, CurrentGraph);
 					}
 				}
 			}
diff --git a/Engine/Source/Editor/UnrealEd/Private/Kismet2/DebuggerCommands.cpp b/Engine/Source/Editor/UnrealEd/Private/Kismet2/DebuggerCommands.cpp
--- a/Engine/Source/Editor/UnrealEd/Private/Kismet2/DebuggerCommands.cpp
+++ b/Engine/Source/Editor/UnrealEd/Private/Kismet2/DebuggerCommands.cpp
@@ -97,6 +97,8 @@ public:
 
 	static void ShowCurrentStatement_Clicked();
 	static void StepInto_Clicked();
+	static void StepOver_Clicked();
+	static void StepOut_Clicked();
 
 	static void TogglePlayPause_Clicked();
 
@@ -248,7 +250,9 @@ void FPlayWorldCommands::RegisterCommands()
 	UI_COMMAND( TogglePlayPauseOfPlaySession, "Toggle Play/Pause", "Resume playing if paused, or pause if playing", EUserInterfaceActionType::Button, FInputChord( EKeys::Pause ) );
 	UI_COMMAND( PossessEjectPlayer, "Possess or Eject Player", "Possesses or ejects the player from the camera", EUserInterfaceActionType::Button, FInputChord( EKeys::F8 ) );
 	UI_COMMAND( ShowCurrentStatement, "Find Node", "Show the current node", EUserInterfaceActionType::Button, FInputChord() );
-	UI_COMMAND( StepInto, "Step", "Step to the next node to be executed", EUserInterfaceActionType::Button, FInputChord( EKeys::F10) );
+	UI_COMMAND( StepInto, "Step Into", "Step Into the next node to be executed", EUserInterfaceActionType::Button, FInputChord( EKeys::F10) );
+	UI_COMMAND( StepOver, "Step Over", "Step to the next node to be executed in the current graph", EUserInterfaceActionType::Button, FInputChord() );
+	UI_COMMAND( StepOut, "Step Out", "Step Out to the next node to be executed in the parent graph", EUserInterfaceActionType::Button, FInputChord() );
 
 	// Launch
 	UI_COMMAND( RepeatLastLaunch, "Launch", "Launches the game on the device as the last session launched from the dropdown next to the Play on Device button on the level editor toolbar", EUserInterfaceActionType::Button, FInputChord( EKeys::P, EModifierKey::Alt | EModifierKey::Shift ) )
@@ -412,6 +416,20 @@ void FPlayWorldCommands::BindGlobalPlayWorldCommands()
 		FIsActionChecked::CreateStatic( &FInternalPlayWorldCommandCallbacks::IsStoppedAtBreakpoint )
 		);
 
+	ActionList.MapAction( Commands.StepOver,
+		FExecuteAction::CreateStatic( &FInternalPlayWorldCommandCallbacks::StepOver_Clicked ),
+		FCanExecuteAction::CreateStatic( &FInternalPlayWorldCommandCallbacks::IsStoppedAtBreakpoint ),
+		FIsActionChecked(),
+		FIsActionChecked::CreateStatic( &FInternalPlayWorldCommandCallbacks::IsStoppedAtBreakpoint )
+		);
+
+	ActionList.MapAction( Commands.StepOut,
+		FExecuteAction::CreateStatic( &FInternalPlayWorldCommandCallbacks::StepOut_Clicked ),
+		FCanExecuteAction::CreateStatic( &FInternalPlayWorldCommandCallbacks::IsStoppedAtBreakpoint ),
+		FIsActionChecked(),
+		FIsActionChecked::CreateStatic( &FInternalPlayWorldCommandCallbacks::IsStoppedAtBreakpoint )
+		);
+
 	ActionList.MapAction( Commands.PlayInNetworkDedicatedServer,
 		FExecuteAction::CreateStatic( &FInternalPlayWorldCommandCallbacks::OnToggleDedicatedServerPIE ),
 		FCanExecuteAction(),
@@ -1698,7 +1716,8 @@ void FInternalPlayWorldCommandCallbacks::PausePlaySession_Clicked()
 void FInternalPlayWorldCommandCallbacks::SingleFrameAdvance_Clicked()
 {
 	// We want to function just like Single stepping where we will stop at a breakpoint if one is encountered but we also want to stop after 1 tick if a breakpoint is not encountered.
-	FKismetDebugUtilities::RequestSingleStepping();
+	const bool bAllowStepIn = true;
+	FKismetDebugUtilities::RequestSingleStepping(bAllowStepIn);
 	if (HasPlayWorld())
 	{
 		GUnrealEd->PlayWorld->bDebugFrameStepExecution = true;
@@ -1730,7 +1749,8 @@ void FInternalPlayWorldCommandCallbacks::ShowCurrentStatement_Clicked()
 
 void FInternalPlayWorldCommandCallbacks::StepInto_Clicked()
 {
-	FKismetDebugUtilities::RequestSingleStepping();
+	const bool bAllowStepIn = true;
+	FKismetDebugUtilities::RequestSingleStepping(bAllowStepIn);
 	if (HasPlayWorld())
 	{
 		LeaveDebuggingMode();
@@ -1738,6 +1758,26 @@ void FInternalPlayWorldCommandCallbacks::StepInto_Clicked()
 	}
 }
 
+void FInternalPlayWorldCommandCallbacks::StepOver_Clicked()
+{
+	const bool bAllowStepIn = false;
+	FKismetDebugUtilities::RequestSingleStepping(bAllowStepIn);
+	if (HasPlayWorld())
+	{
+		LeaveDebuggingMode();
+		GUnrealEd->PlaySessionSingleStepped();
+	}
+}
+
+void FInternalPlayWorldCommandCallbacks::StepOut_Clicked()
+{
+	FKismetDebugUtilities::RequestStepOut();
+	if (HasPlayWorld())
+	{
+		LeaveDebuggingMode();
+		GUnrealEd->PlaySessionSingleStepped();
+	}
+}
 
 void FInternalPlayWorldCommandCallbacks::TogglePlayPause_Clicked()
 {
diff --git a/Engine/Source/Editor/UnrealEd/Private/Kismet2/KismetDebugUtilities.cpp b/Engine/Source/Editor/UnrealEd/Private/Kismet2/KismetDebugUtilities.cpp
--- a/Engine/Source/Editor/UnrealEd/Private/Kismet2/KismetDebugUtilities.cpp
+++ b/Engine/Source/Editor/UnrealEd/Private/Kismet2/KismetDebugUtilities.cpp
@@ -21,6 +21,9 @@
 
 TWeakObjectPtr< UEdGraphNode > FKismetDebugUtilities::CurrentInstructionPointer;
 TWeakObjectPtr< UEdGraphNode > FKismetDebugUtilities::MostRecentBreakpointInstructionPointer;
+int32 FKismetDebugUtilities::TargetGraphStackDepth = INDEX_NONE;
+TArray<TWeakObjectPtr<UEdGraph>> FKismetDebugUtilities::GraphStack;
+
 FString FKismetDebugUtilities::LastExceptionMessage;
 const FFrame* FKismetDebugUtilities::StackFrameAtIntraframeDebugging = NULL;
 
@@ -30,12 +33,30 @@ bool FKismetDebugUtilities::bIsSingleStepping = false;
 
 //////////////////////////////////////////////////////////////////////////
 
-void FKismetDebugUtilities::RequestSingleStepping()
+void FKismetDebugUtilities::EndOfScriptExecution()
 {
-	bIsSingleStepping = true;
+	bIsSingleStepping = false;
+	TargetGraphStackDepth = INDEX_NONE;
+	GraphStack.SetNum(0, false);
 }
 
+void FKismetDebugUtilities::RequestSingleStepping(bool bInAllowStepIn)
+{
+	bIsSingleStepping = bInAllowStepIn;
+	if (!bInAllowStepIn)
+	{
+		TargetGraphStackDepth = GraphStack.Num();
+	}
+}
 
+void FKismetDebugUtilities::RequestStepOut()
+{
+	bIsSingleStepping = false;
+	if (GraphStack.Num() > 1)
+	{
+		TargetGraphStackDepth = GraphStack.Num() - 1;
+	}
+}
 
 void FKismetDebugUtilities::OnScriptException(const UObject* ActiveObject, const FFrame& StackFrame, const FBlueprintExceptionInfo& Info)
 {
@@ -197,6 +218,11 @@ void FKismetDebugUtilities::OnScriptException(const UObject* ActiveObject, const
 			UEdGraphNode* NodeStoppedAt = FindSourceNodeForCodeLocation(ActiveObject, StackFrame.Node, BreakpointOffset, /*bAllowImpreciseHit=*/ true);
 			if (NodeStoppedAt)
 			{
+				// Handle Node stepping and update the stack
+				if (Info.GetType() == EBlueprintExceptionType::Tracepoint || Info.GetType() == EBlueprintExceptionType::Breakpoint)
+				{
+					CheckBreakConditions(NodeStoppedAt, bShouldBreakExecution);
+				}
 				// If the code which we hit was generated by a macro node expansion
 				UK2Node_MacroInstance* MacroInstanceNode = Cast<UK2Node_MacroInstance>(NodeStoppedAt);
 				if (MacroInstanceNode)
@@ -420,6 +446,44 @@ UEdGraphNode* FKismetDebugUtilities::FindMacroSourceNodeForCodeLocation(const UO
 	return NULL;
 }
 
+void FKismetDebugUtilities::CheckBreakConditions(UEdGraphNode* NodeStoppedAt, bool& InOutBreakExecution)
+{
+	if (NodeStoppedAt)
+	{
+		// Update tracked graph stack.
+		if (!GraphStack.Num())
+		{
+			GraphStack.Push(NodeStoppedAt->GetTypedOuter<UEdGraph>());
+		}
+		else if (UK2Node_FunctionEntry* FunctionNode = Cast<UK2Node_FunctionEntry>(NodeStoppedAt))
+		{
+			UEdGraph* NewGraph = FunctionNode->GetTypedOuter<UEdGraph>();
+			const bool bAddFunction = GraphStack.Num() && NewGraph ? NewGraph != GraphStack.Last() : false;
+
+			if (bAddFunction)
+			{
+				GraphStack.Push(NewGraph);
+			}
+		}
+		else
+		{
+			UEdGraph* CurrGraph = NodeStoppedAt->GetTypedOuter<UEdGraph>();
+			if (GraphStack.Last() != CurrGraph)
+			{
+				GraphStack.Pop();
+			}
+		}
+		// Figure out if we have a break condition.
+		if (!bIsSingleStepping)
+		{
+			if (TargetGraphStackDepth != INDEX_NONE)
+			{
+				InOutBreakExecution = TargetGraphStackDepth >= GraphStack.Num();
+			}
+		}
+	}
+}
+
 void FKismetDebugUtilities::AttemptToBreakExecution(UBlueprint* BlueprintObj, const UObject* ActiveObject, const FFrame& StackFrame, const FBlueprintExceptionInfo& Info, UEdGraphNode* NodeStoppedAt, int32 DebugOpcodeOffset)
 {
 	checkSlow(BlueprintObj->GetObjectBeingDebugged() == ActiveObject);
@@ -480,6 +544,7 @@ void FKismetDebugUtilities::AttemptToBreakExecution(UBlueprint* BlueprintObj, co
 		// Pause the simulation
 		GUnrealEd->PlayWorld->bDebugPauseExecution = true;
 		GUnrealEd->PlayWorld->bDebugFrameStepExecution = false;
+		bShouldInStackDebug = true;
 	}
 	else
 	{
diff --git a/Engine/Source/Editor/UnrealEd/Private/Kismet2/KismetReinstanceUtilities.cpp b/Engine/Source/Editor/UnrealEd/Private/Kismet2/KismetReinstanceUtilities.cpp
--- a/Engine/Source/Editor/UnrealEd/Private/Kismet2/KismetReinstanceUtilities.cpp
+++ b/Engine/Source/Editor/UnrealEd/Private/Kismet2/KismetReinstanceUtilities.cpp
@@ -131,7 +131,7 @@ struct FReplaceReferenceHelper
 TSet<TWeakObjectPtr<UBlueprint>> FBlueprintCompileReinstancer::DependentBlueprintsToRefresh = TSet<TWeakObjectPtr<UBlueprint>>();
 TSet<TWeakObjectPtr<UBlueprint>> FBlueprintCompileReinstancer::DependentBlueprintsToRecompile = TSet<TWeakObjectPtr<UBlueprint>>();
 TSet<TWeakObjectPtr<UBlueprint>> FBlueprintCompileReinstancer::DependentBlueprintsToByteRecompile = TSet<TWeakObjectPtr<UBlueprint>>();
-TSet<UBlueprint*> FBlueprintCompileReinstancer::CompiledBlueprintsToSave = TSet<UBlueprint*>();
+TSet<TWeakObjectPtr<UBlueprint>> FBlueprintCompileReinstancer::CompiledBlueprintsToSave = TSet<TWeakObjectPtr<UBlueprint>>();
 
 UClass* FBlueprintCompileReinstancer::HotReloadedOldClass = nullptr;
 UClass* FBlueprintCompileReinstancer::HotReloadedNewClass = nullptr;
@@ -143,7 +143,7 @@ FBlueprintCompileReinstancer::FBlueprintCompileReinstancer(UClass* InClassToRein
 	, bHasReinstanced(false)
 	, bSkipGarbageCollection(bSkipGC)
 	, ClassToReinstanceDefaultValuesCRC(0)
-	, bIsSourceReinstancer(false)
+	, bIsRootReinstancer(false)
 {
 	if( InClassToReinstance != NULL )
 	{
@@ -262,14 +262,18 @@ FBlueprintCompileReinstancer::FBlueprintCompileReinstancer(UClass* InClassToRein
 			ClassToReinstanceDefaultValuesCRC = GeneratingBP->CrcLastCompiledCDO;
 			FBlueprintEditorUtils::GetDependentBlueprints(GeneratingBP, Dependencies);
 
-			bool const bIsLevelPackage = (UWorld::FindWorldInPackage(GeneratingBP->GetOutermost()) != nullptr);
-			// we don't want to save the entire level (especially if this 
-			// compile was already kicked off as a result of a level save, as it
-			// could cause a recursive save)... let the "SaveOnCompile" setting 
-			// only save blueprint assets
-			if (!bIsLevelPackage)
+			// Never queue for saving when regenerating on load
+			if (!GeneratingBP->bIsRegeneratingOnLoad && !bIsReinstancingSkeleton)
 			{
-				CompiledBlueprintsToSave.Add(GeneratingBP);
+				bool const bIsLevelPackage = (UWorld::FindWorldInPackage(GeneratingBP->GetOutermost()) != nullptr);
+				// we don't want to save the entire level (especially if this 
+				// compile was already kicked off as a result of a level save, as it
+				// could cause a recursive save)... let the "SaveOnCompile" setting 
+				// only save blueprint assets
+				if (!bIsLevelPackage)
+				{
+					CompiledBlueprintsToSave.Add(GeneratingBP);
+				}
 			}
 		}
 	}
@@ -337,21 +341,26 @@ void FBlueprintCompileReinstancer::OptionallyRefreshNodes(UBlueprint* CurrentBP)
 
 FBlueprintCompileReinstancer::~FBlueprintCompileReinstancer()
 {
-	if (bIsSourceReinstancer)
+	if (bIsRootReinstancer)
 	{
 		if (CompiledBlueprintsToSave.Num() > 0)
 		{
 			if ( !IsRunningCommandlet() && !GIsAutomationTesting )
 			{
 				TArray<UPackage*> PackagesToSave;
-				for (UBlueprint* BP : CompiledBlueprintsToSave)
+				for (TWeakObjectPtr<UBlueprint> BPPtr : CompiledBlueprintsToSave)
 				{
-					UBlueprintEditorSettings* Settings = GetMutableDefault<UBlueprintEditorSettings>();
-					const bool bShouldSaveOnCompile = (( Settings->SaveOnCompile == SoC_Always ) || ( ( Settings->SaveOnCompile == SoC_SuccessOnly ) && ( BP->Status == BS_UpToDate ) ));
-
-					if (bShouldSaveOnCompile)
+					if (BPPtr.IsValid())
 					{
-						PackagesToSave.Add(BP->GetOutermost());
+						UBlueprint* BP = BPPtr.Get();
+
+						UBlueprintEditorSettings* Settings = GetMutableDefault<UBlueprintEditorSettings>();
+						const bool bShouldSaveOnCompile = ((Settings->SaveOnCompile == SoC_Always) || ((Settings->SaveOnCompile == SoC_SuccessOnly) && (BP->Status == BS_UpToDate)));
+
+						if (bShouldSaveOnCompile)
+						{
+							PackagesToSave.Add(BP->GetOutermost());
+						}
 					}
 				}
 
@@ -593,7 +602,7 @@ void FBlueprintCompileReinstancer::ReinstanceObjects(bool bForceAlwaysReinstance
 		if (QueueToReinstance.Num() && (QueueToReinstance[0] == SharedThis))
 		{
 			// Mark it as the source reinstancer, no other reinstancer can get here until this Blueprint finishes compiling
-			bIsSourceReinstancer = true;
+			bIsRootReinstancer = true;
 
 			TSet<TWeakObjectPtr<UBlueprint>> CompiledBlueprints;
 
diff --git a/Engine/Source/Editor/UnrealEd/Private/UnrealEdMisc.cpp b/Engine/Source/Editor/UnrealEd/Private/UnrealEdMisc.cpp
--- a/Engine/Source/Editor/UnrealEd/Private/UnrealEdMisc.cpp
+++ b/Engine/Source/Editor/UnrealEd/Private/UnrealEdMisc.cpp
@@ -200,6 +200,7 @@ void FUnrealEdMisc::OnInit()
 
 	/** Delegate that gets called when a script exception occurs */
 	FBlueprintCoreDelegates::OnScriptException.AddStatic(&FKismetDebugUtilities::OnScriptException);
+	FBlueprintCoreDelegates::OnScriptExecutionEnd.AddStatic(&FKismetDebugUtilities::EndOfScriptExecution);
 	
 	FEditorDelegates::ChangeEditorMode.AddRaw(this, &FUnrealEdMisc::OnEditorChangeMode);
 	FCoreDelegates::PreModal.AddRaw(this, &FUnrealEdMisc::OnEditorPreModal);
diff --git a/Engine/Source/Editor/UnrealEd/Public/GraphEditor.h b/Engine/Source/Editor/UnrealEd/Public/GraphEditor.h
--- a/Engine/Source/Editor/UnrealEd/Public/GraphEditor.h
+++ b/Engine/Source/Editor/UnrealEd/Public/GraphEditor.h
@@ -396,6 +396,15 @@ public:
 		}
 	}
 
+	/** Capture keyboard */
+	virtual void CaptureKeyboard()
+	{
+		if (Implementation.IsValid())
+		{
+			Implementation->CaptureKeyboard();
+		}
+	}
+
 protected:
 	/** Invoked when the underlying Graph is being changed. */
 	virtual void OnGraphChanged(const struct FEdGraphEditAction& InAction)
diff --git a/Engine/Source/Editor/UnrealEd/Public/Kismet2/BlueprintEditorUtils.h b/Engine/Source/Editor/UnrealEd/Public/Kismet2/BlueprintEditorUtils.h
--- a/Engine/Source/Editor/UnrealEd/Public/Kismet2/BlueprintEditorUtils.h
+++ b/Engine/Source/Editor/UnrealEd/Public/Kismet2/BlueprintEditorUtils.h
@@ -1054,6 +1054,14 @@ public:
 	/** Remove an implemented interface, and its associated member function graphs.  If bPreserveFunctions is true, then the interface will move its functions to be normal implemented blueprint functions */
 	static void RemoveInterface(UBlueprint* Blueprint, const FName& InterfaceClassName, bool bPreserveFunctions = false);
 
+	/**
+	* Promotes a Graph from being an Interface Override to a full member function
+	*
+	* @param InBlueprint			Blueprint the graph is contained within
+	* @param InInterfaceGraph		The graph to promote
+	*/
+	static void PromoteGraphFromInterfaceOverride(UBlueprint* InBlueprint, UEdGraph* InInterfaceGraph);
+
 	/** Gets the graphs currently in the blueprint associated with the specified interface */
 	static void GetInterfaceGraphs(UBlueprint* Blueprint, const FName& InterfaceClassName, TArray<UEdGraph*>& ChildGraphs);
 
@@ -1367,6 +1375,8 @@ struct UNREALED_API FBlueprintDuplicationScopeFlags
 		NoFlags = 0,
 		NoExtraCompilation = 1 << 0,
 		TheSameTimelineGuid = 1 << 1,
+		// This flag is needed for C++ backend (while compiler validates graphs). The actual BPGC type is compatible with the original BPGC.
+		ValidatePinsUsingSourceClass = 1 << 2,
 	};
 
 	static uint32 bStaticFlags;
diff --git a/Engine/Source/Editor/UnrealEd/Public/Kismet2/DebuggerCommands.h b/Engine/Source/Editor/UnrealEd/Public/Kismet2/DebuggerCommands.h
--- a/Engine/Source/Editor/UnrealEd/Public/Kismet2/DebuggerCommands.h
+++ b/Engine/Source/Editor/UnrealEd/Public/Kismet2/DebuggerCommands.h
@@ -63,6 +63,8 @@ public:
 	TSharedPtr<FUICommandInfo> PossessEjectPlayer;
 	TSharedPtr<FUICommandInfo> ShowCurrentStatement;
 	TSharedPtr<FUICommandInfo> StepInto;
+	TSharedPtr<FUICommandInfo> StepOver;
+	TSharedPtr<FUICommandInfo> StepOut;
 
 	/** Launch on device */
 	TSharedPtr<FUICommandInfo> RepeatLastLaunch;
diff --git a/Engine/Source/Editor/UnrealEd/Public/Kismet2/KismetDebugUtilities.h b/Engine/Source/Editor/UnrealEd/Public/Kismet2/KismetDebugUtilities.h
--- a/Engine/Source/Editor/UnrealEd/Public/Kismet2/KismetDebugUtilities.h
+++ b/Engine/Source/Editor/UnrealEd/Public/Kismet2/KismetDebugUtilities.h
@@ -115,8 +115,14 @@ public:
 	/** Returns the most recent hit breakpoint; if a PIE/SIE session is started but paused; otherwise NULL */
 	static class UEdGraphNode* GetMostRecentBreakpointHit();
 
-	/** Request an attempt to single-step to the next node */
-	static void RequestSingleStepping();
+	/** Request an attempt to single-step to the next node, with parameter to control step into sub graphs */
+	static void RequestSingleStepping(bool bInAllowStepIn);
+
+	/** Request an attempt to step out of the current graph */
+	static void RequestStepOut();
+
+	/** Called on terminatation of the current script execution so we can reset any break conditions */
+	static void EndOfScriptExecution();
 
 	// The maximum number of trace samples to gather before overwriting old ones
 	enum { MAX_TRACE_STACK_SAMPLES = 1024 };
@@ -219,6 +225,7 @@ public:
 	// This doesn't work properly if there is more than one blueprint editor open at once either (one will consume it, the others will be left in the cold)
 	static FText GetAndClearLastExceptionMessage();
 protected:
+	static void CheckBreakConditions(UEdGraphNode* NodeStoppedAt, bool& InOutBreakExecution);
 	static void AttemptToBreakExecution(UBlueprint* BlueprintObj, const UObject* ActiveObject, const FFrame& StackFrame, const FBlueprintExceptionInfo& Info, UEdGraphNode* NodeStoppedAt, int32 DebugOpcodeOffset);
 
 private:
@@ -227,6 +234,12 @@ private:
 	// The current instruction encountered if we are stopped at a breakpoint; NULL otherwise
 	static TWeakObjectPtr< class UEdGraphNode > MostRecentBreakpointInstructionPointer;
 
+	// The current function call graph stack.
+	static TArray<TWeakObjectPtr<class UEdGraph>> GraphStack;
+
+	// The target graph call stack depth. INDEX_NONE if not active.
+	static int32 TargetGraphStackDepth;
+
 	// The last message that an exception delivered
 	static FString LastExceptionMessage;
 
diff --git a/Engine/Source/Editor/UnrealEd/Public/Kismet2/KismetReinstanceUtilities.h b/Engine/Source/Editor/UnrealEd/Public/Kismet2/KismetReinstanceUtilities.h
--- a/Engine/Source/Editor/UnrealEd/Public/Kismet2/KismetReinstanceUtilities.h
+++ b/Engine/Source/Editor/UnrealEd/Public/Kismet2/KismetReinstanceUtilities.h
@@ -40,7 +40,7 @@ protected:
 	static TSet<TWeakObjectPtr<UBlueprint>> DependentBlueprintsToRefresh;
 	static TSet<TWeakObjectPtr<UBlueprint>> DependentBlueprintsToRecompile;
 	static TSet<TWeakObjectPtr<UBlueprint>> DependentBlueprintsToByteRecompile;
-	static TSet<UBlueprint*> CompiledBlueprintsToSave;
+	static TSet<TWeakObjectPtr<UBlueprint>> CompiledBlueprintsToSave;
 
 	static UClass* HotReloadedOldClass;
 	static UClass* HotReloadedNewClass;
@@ -78,8 +78,8 @@ protected:
 	/** Objects that should keep reference to old class */
 	TSet<UObject*> ObjectsThatShouldUseOldStuff;
 
-	/** TRUE if this is the source reinstancer that all other active reinstancing is spawned from */
-	bool bIsSourceReinstancer;
+	/** TRUE if this is the root reinstancer that all other active reinstancing is spawned from */
+	bool bIsRootReinstancer;
 
 public:
 	// FSerializableObject interface
@@ -147,7 +147,9 @@ protected:
 		, OriginalCDO(NULL)
 		, bHasReinstanced(false)
 		, bSkipGarbageCollection(false)
+		, bIsReinstancingSkeleton(false)
 		, ClassToReinstanceDefaultValuesCRC(0)
+		, bIsRootReinstancer(false)
 	{}
 
 	/** Sets the reinstancer up to work on every object of the specified class */
diff --git a/Engine/Source/Editor/UnrealEd/UnrealEd.Build.cs b/Engine/Source/Editor/UnrealEd/UnrealEd.Build.cs
--- a/Engine/Source/Editor/UnrealEd/UnrealEd.Build.cs
+++ b/Engine/Source/Editor/UnrealEd/UnrealEd.Build.cs
@@ -62,6 +62,7 @@ public class UnrealEd : ModuleRules
                 "MessagingRpc",
                 "PortalRpc",
                 "PortalServices",
+                "BlueprintNativeCodeGen",
 			}
 		);
 
@@ -144,7 +145,6 @@ public class UnrealEd : ModuleRules
 				"AddContentDialog",
 				"GameProjectGeneration",
                 "HierarchicalLODUtilities",
-                "BlueprintNativeCodeGen",
                 "Analytics"
 			}
 		);
@@ -216,6 +216,7 @@ public class UnrealEd : ModuleRules
                 "PortalProxies",
                 "PortalServices",
                 "GeometryCacheEd",
+                "BlueprintNativeCodeGen",
 			}
 		);
 
diff --git a/Engine/Source/Programs/AutomationTool/AutomationUtils/ProjectParams.cs b/Engine/Source/Programs/AutomationTool/AutomationUtils/ProjectParams.cs
--- a/Engine/Source/Programs/AutomationTool/AutomationUtils/ProjectParams.cs
+++ b/Engine/Source/Programs/AutomationTool/AutomationUtils/ProjectParams.cs
@@ -989,15 +989,7 @@ namespace AutomationTool
 		/// <summary>
 		/// Shared: Ref to an auto-generated plugin file that should be incorporated into the project's build
 		/// </summary>
-		public FileReference NativizedScriptPlugin;
-
-		/// <summary>
-		/// Shared: Used to guard against fatal use of the GeneratedScriptPlugin param
-		/// </summary>
-		public bool UseNativizedScriptPlugin()
-		{
-			return (NativizedScriptPlugin != null) && NativizedScriptPlugin.Exists();
-		}
+        public List<FileReference> BlueprintPluginPaths = new List<FileReference>();
 
 		#endregion
 
diff --git a/Engine/Source/Programs/AutomationTool/Scripts/BuildProjectCommand.Automation.cs b/Engine/Source/Programs/AutomationTool/Scripts/BuildProjectCommand.Automation.cs
--- a/Engine/Source/Programs/AutomationTool/Scripts/BuildProjectCommand.Automation.cs
+++ b/Engine/Source/Programs/AutomationTool/Scripts/BuildProjectCommand.Automation.cs
@@ -82,9 +82,13 @@ public partial class Project : CommandUtils
         string ScriptPluginArgs = "";
         // if we're utilizing an auto-generated code plugin/module (a product of 
         // the cook process), make sure to compile it along with the targets here
-        if (Params.UseNativizedScriptPlugin())
+        if (Params.RunAssetNativization)
         {
-            ScriptPluginArgs = "-PLUGIN \"" + Params.NativizedScriptPlugin + "\"";
+            // Add every plugin:
+            foreach( var CodePlugin in  Params.BlueprintPluginPaths)
+            {
+                ScriptPluginArgs = "-PLUGIN \"" + CodePlugin + "\""; 
+            }
         }
 
 		// Setup cooked targets
diff --git a/Engine/Source/Programs/AutomationTool/Scripts/CookCommand.Automation.cs b/Engine/Source/Programs/AutomationTool/Scripts/CookCommand.Automation.cs
--- a/Engine/Source/Programs/AutomationTool/Scripts/CookCommand.Automation.cs
+++ b/Engine/Source/Programs/AutomationTool/Scripts/CookCommand.Automation.cs
@@ -219,17 +219,17 @@ public partial class Project : CommandUtils
                 // into native source code... the cooker needs to 
                 if (Params.RunAssetNativization)
                 {
-                    string GeneratedPluginName = "NativizedScript";
-                    string ProjectDir = Params.RawProjectPath.Directory.ToString();
-                    string GeneratedPluginDirectory = Path.GetFullPath(CommandUtils.CombinePaths(ProjectDir, "Intermediate", "Plugins", GeneratedPluginName));
-                    string PluginDescFilename  = GeneratedPluginName + ".uplugin";
-                    string GeneratedPluginPath = Path.GetFullPath(CommandUtils.CombinePaths(GeneratedPluginDirectory, PluginDescFilename));
-
-                    CommandletParams += " -NativizeAssets -NativizedAssetPlugin=\"" + GeneratedPluginPath + "\"";
+                    CommandletParams += " -NativizeAssets";
 
-                    // fill out this, so as to coordinate with build automation
-                    // (the builder now needs to compile in the generated assets)
-                    Params.NativizedScriptPlugin = new FileReference(GeneratedPluginPath);
+                    // Store plugin paths now, it's easiest to do so while PlatformsToCook is still available:
+                    string ProjectDir = Params.RawProjectPath.Directory.ToString();
+                    foreach (var Platform in PlatformsToCook)
+                    {
+                        // If you change this target path you must also update logic in CookOnTheFlyServer.cpp. Passing a single directory around is cumbersome for testing, so I have hard coded it.
+                        // Similarly if you change the .uplugin name you must update DefaultPluginName in BlueprintNativeCodeGenModule.cpp
+                        string GeneratedPluginPath = CombinePaths(GetDirectoryName(ProjectDir), "Intermediate", Platform, "NativizedAssets/NativizedAssets.uplugin");
+                        Params.BlueprintPluginPaths.Add( new FileReference(GeneratedPluginPath) );
+                    }
                 }
                 if (Params.HasAdditionalCookerOptions)
                 {
diff --git a/Engine/Source/Programs/UnrealHeaderTool/Private/CodeGenerator.cpp b/Engine/Source/Programs/UnrealHeaderTool/Private/CodeGenerator.cpp
--- a/Engine/Source/Programs/UnrealHeaderTool/Private/CodeGenerator.cpp
+++ b/Engine/Source/Programs/UnrealHeaderTool/Private/CodeGenerator.cpp
@@ -4600,7 +4600,7 @@ TArray<UFunction*> FNativeClassHeaderGenerator::ExportCallbackFunctions(FUnrealS
 		FString FunctionName = Function->GetName();
 
 		GeneratedHeaderText.Logf(TEXT("extern %s FName %s_%s;") LINE_TERMINATOR, *GetAPIString(), *API, *FunctionName);
-		ReferencedNames.Add(Function->GetFName());
+		ReferencedNames.Add(Function->GetFName(), GetOverriddenFName(Function));
 
 		const FFuncInfo& FunctionData = CompilerInfo->GetFunctionData();
 
@@ -5304,18 +5304,14 @@ void FNativeClassHeaderGenerator::ExportGeneratedCPP()
 
 	GeneratedCPPText.Log(*GeneratedPackageCPP);
 
-	// Add all names marked for export to a list (for sorting)
-	TArray<FString> Names;
-	for (TSet<FName>::TIterator It(ReferencedNames); It; ++It)
 	{
-		Names.Add(It->ToString());
-	}	
+		// Autogenerate names (alphabetically sorted).
+		ReferencedNames.KeySort( TLess<FName>() );
 
-	// Autogenerate names (alphabetically sorted).
-	Names.Sort();
-	for( int32 NameIndex=0; NameIndex<Names.Num(); NameIndex++ )
-	{
-		GeneratedCPPText.Logf( TEXT("FName %s_%s = FName(TEXT(\"%s\"));") LINE_TERMINATOR, *API, *Names[NameIndex], *Names[NameIndex] );
+		for (auto& PairIt : ReferencedNames)
+		{
+			GeneratedCPPText.Logf(TEXT("FName %s_%s = FName(TEXT(\"%s\"));") LINE_TERMINATOR, *API, *PairIt.Key.ToString(), *PairIt.Value.ToString());
+		}
 	}
 
 	GeneratedCPPEpilogue.Logf(
diff --git a/Engine/Source/Programs/UnrealHeaderTool/Private/HeaderParser.cpp b/Engine/Source/Programs/UnrealHeaderTool/Private/HeaderParser.cpp
--- a/Engine/Source/Programs/UnrealHeaderTool/Private/HeaderParser.cpp
+++ b/Engine/Source/Programs/UnrealHeaderTool/Private/HeaderParser.cpp
@@ -2423,9 +2423,14 @@ void FHeaderParser::FixupDelegateProperties( FClasses& AllClasses, UStruct* Stru
 							for (TFieldIterator<UProperty> PropIt(SourceDelegateFunction); PropIt && (PropIt->PropertyFlags & CPF_Parm); ++PropIt)
 							{
 								UProperty* FuncParam = *PropIt;
-								if(FuncParam->HasAllPropertyFlags(CPF_OutParm) && !FuncParam->HasAllPropertyFlags(CPF_ConstParm))
+								if(FuncParam->HasAllPropertyFlags(CPF_OutParm) && !FuncParam->HasAllPropertyFlags(CPF_ConstParm)  )
 								{
-									FError::Throwf(TEXT("BlueprintAssignable delegates do not support non-const references at the moment. Function: %s Parameter: '%s'"), *SourceDelegateFunction->GetName(), *FuncParam->GetName());
+									const bool bClassGeneratedFromBP = FClass::IsDynamic(Struct);
+									const bool bAllowedArrayRefFromBP = bClassGeneratedFromBP && FuncParam->IsA<UArrayProperty>();
+									if (!bAllowedArrayRefFromBP)
+									{
+										FError::Throwf(TEXT("BlueprintAssignable delegates do not support non-const references at the moment. Function: %s Parameter: '%s'"), *SourceDelegateFunction->GetName(), *FuncParam->GetName());
+									}
 								}
 							}
 						}
@@ -5827,7 +5832,12 @@ void FHeaderParser::CompileFunctionDeclaration(FClasses& AllClasses)
 	}
 	if ( (bHasAnyOutputs == false) && (FuncInfo.FunctionFlags & (FUNC_BlueprintPure)) )
 	{
-		FError::Throwf(TEXT("BlueprintPure specifier is not allowed for functions with no return value and no output parameters."));
+		// This bad behavior would be treated as a warning in the Blueprint editor, so when converted assets generates these bad functions
+		// we don't want to prevent compilation:
+		if (!bClassGeneratedFromBP)
+		{
+			FError::Throwf(TEXT("BlueprintPure specifier is not allowed for functions with no return value and no output parameters."));
+		}
 	}
 
 
diff --git a/Engine/Source/Programs/UnrealHeaderTool/Private/NativeClassExporter.h b/Engine/Source/Programs/UnrealHeaderTool/Private/NativeClassExporter.h
--- a/Engine/Source/Programs/UnrealHeaderTool/Private/NativeClassExporter.h
+++ b/Engine/Source/Programs/UnrealHeaderTool/Private/NativeClassExporter.h
@@ -89,8 +89,8 @@ private:
 	/** Set of already exported cross-module references, to prevent duplicates */
 	TSet<FString> UniqueCrossModuleReferences;
 
-	/** Names that have been exported so far **/
-	TSet<FName> ReferencedNames;
+	/** Names that have been exported so far. Key is a native name, value is a overriden name (or also the native)**/
+	TMap<FName, FName> ReferencedNames;
 
 	/** the existing disk version of the header for this package's names */
 	FString OriginalNamesHeader;
diff --git a/Engine/Source/Runtime/AIModule/Classes/EnvironmentQuery/EnvQueryInstanceBlueprintWrapper.h b/Engine/Source/Runtime/AIModule/Classes/EnvironmentQuery/EnvQueryInstanceBlueprintWrapper.h
--- a/Engine/Source/Runtime/AIModule/Classes/EnvironmentQuery/EnvQueryInstanceBlueprintWrapper.h
+++ b/Engine/Source/Runtime/AIModule/Classes/EnvironmentQuery/EnvQueryInstanceBlueprintWrapper.h
@@ -11,6 +11,7 @@ class AIMODULE_API UEnvQueryInstanceBlueprintWrapper : public UObject, public IE
 {
 	GENERATED_BODY()
 
+public:
 	DECLARE_DYNAMIC_MULTICAST_DELEGATE_TwoParams(FEQSQueryDoneSignature, UEnvQueryInstanceBlueprintWrapper*, QueryInstance, EEnvQueryStatus::Type, QueryStatus);
 
 protected:
diff --git a/Engine/Source/Runtime/Core/Private/Tests/Math/UnrealMathTest.cpp b/Engine/Source/Runtime/Core/Private/Tests/Math/UnrealMathTest.cpp
--- a/Engine/Source/Runtime/Core/Private/Tests/Math/UnrealMathTest.cpp
+++ b/Engine/Source/Runtime/Core/Private/Tests/Math/UnrealMathTest.cpp
@@ -1375,3 +1375,91 @@ bool FVectorRegisterAbstractionTest::RunTest(const FString& Parameters)
 
 	return true;
 }
+
+IMPLEMENT_SIMPLE_AUTOMATION_TEST(FInterpolationFunctionTests, "System.Core.Math.Interpolation Function Test", EAutomationTestFlags::ApplicationContextMask | EAutomationTestFlags::SmokeFilter)
+bool FInterpolationFunctionTests::RunTest(const FString&)
+{
+	// The purpose of this test is to verify that various combinations of the easing functions are actually equivalent.
+	// It currently only tests the InOut versions over different ranges, because the initial implementation was bad.
+	// Further improvements (optimizations, new easing functions) to the easing functions should be accompanied by
+	// expansions to this test suite.
+
+	typedef float(*EasingFunc)(float Percent);
+	auto RunInOutTest = [](const TArray< TPair<EasingFunc, FString> >& Functions, FAutomationTestBase* TestContext)
+	{
+		for (int32 I = 0; I < 100; ++I)
+		{
+			float Percent = (float)I / 100.f;
+			TArray<float> Values;
+			for (const auto& Entry : Functions)
+			{
+				Values.Push(Entry.Key(Percent));
+			}
+
+			bool bSucceeded = true;
+			int32 K = 0;
+			for (int32 J = 1; J < Functions.Num(); ++J)
+			{
+				if (!FMath::IsNearlyEqual(Values[K], Values[J], 0.0001f))
+				{
+					TestContext->AddError(FString::Printf( TEXT("Easing Function tests failed at index %d!"), I ) );
+
+					for (int32 L = 0; L < Values.Num(); ++L)
+					{
+						TestContext->AddLogItem(FString::Printf(TEXT("%s: %f"), *(Functions[L].Value), Values[L]));
+					}
+					// don't record further failures, it would likely create a tremendous amount of spam
+					return;
+				}
+			}
+		}
+	};
+
+#define INTERP_WITH_RANGE( RANGE_MIN, RANGE_MAX, FUNCTION, IDENTIFIER ) \
+	auto FUNCTION##IDENTIFIER = [](float Percent ) \
+	{ \
+		const float Min = RANGE_MIN; \
+		const float Max = RANGE_MAX; \
+		const float Range = Max - Min; \
+		return (FMath::FUNCTION(Min, Max, Percent) - Min) / Range; \
+	};
+
+	{
+		// Test InterpExpoInOut:
+		INTERP_WITH_RANGE(.9f, 1.2f, InterpExpoInOut, A)
+		INTERP_WITH_RANGE(0.f, 1.f, InterpExpoInOut, B)
+		INTERP_WITH_RANGE(-8.6f;, 2.3f, InterpExpoInOut, C)
+		TArray< TPair< EasingFunc, FString > > FunctionsToTest;
+		FunctionsToTest.Push(TPairInitializer<EasingFunc, FString>(InterpExpoInOutA, TEXT("InterpExpoInOutA")));
+		FunctionsToTest.Push(TPairInitializer<EasingFunc, FString>(InterpExpoInOutB, TEXT("InterpExpoInOutB")));
+		FunctionsToTest.Push(TPairInitializer<EasingFunc, FString>(InterpExpoInOutC, TEXT("InterpExpoInOutC")));
+		RunInOutTest(FunctionsToTest, this);
+	}
+
+	{
+		// Test InterpCircularInOut:
+		INTERP_WITH_RANGE(5.f, 9.32f, InterpCircularInOut, A)
+		INTERP_WITH_RANGE(0.f, 1.f, InterpCircularInOut, B)
+		INTERP_WITH_RANGE(-8.1f;, -.75f, InterpCircularInOut, C)
+		TArray< TPair< EasingFunc, FString > > FunctionsToTest;
+		FunctionsToTest.Push(TPairInitializer<EasingFunc, FString>(InterpCircularInOutA, TEXT("InterpCircularInOutA")));
+		FunctionsToTest.Push(TPairInitializer<EasingFunc, FString>(InterpCircularInOutB, TEXT("InterpCircularInOutB")));
+		FunctionsToTest.Push(TPairInitializer<EasingFunc, FString>(InterpCircularInOutC, TEXT("InterpCircularInOutC")));
+		RunInOutTest(FunctionsToTest, this);
+	}
+
+	{
+		// Test InterpSinInOut:
+		INTERP_WITH_RANGE(10.f, 11.2f, InterpSinInOut, A)
+		INTERP_WITH_RANGE(0.f, 1.f, InterpSinInOut, B)
+		INTERP_WITH_RANGE(-5.6f;, -4.3f, InterpSinInOut, C)
+		TArray< TPair< EasingFunc, FString > > FunctionsToTest;
+		FunctionsToTest.Push(TPairInitializer<EasingFunc, FString>(InterpSinInOutA, TEXT("InterpSinInOutA")));
+		FunctionsToTest.Push(TPairInitializer<EasingFunc, FString>(InterpSinInOutB, TEXT("InterpSinInOutB")));
+		FunctionsToTest.Push(TPairInitializer<EasingFunc, FString>(InterpSinInOutC, TEXT("InterpSinInOutC")));
+		RunInOutTest(FunctionsToTest, this);
+	}
+
+	return true;
+}
+
diff --git a/Engine/Source/Runtime/Core/Public/Math/UnrealMathUtility.h b/Engine/Source/Runtime/Core/Public/Math/UnrealMathUtility.h
--- a/Engine/Source/Runtime/Core/Public/Math/UnrealMathUtility.h
+++ b/Engine/Source/Runtime/Core/Public/Math/UnrealMathUtility.h
@@ -663,7 +663,7 @@ struct FMath : public FPlatformMath
 	template< class T >
 	static FORCEINLINE_DEBUGGABLE T InterpEaseIn(const T& A, const T& B, float Alpha, float Exp)
 	{
-		float const ModifiedAlpha = FMath::Pow(Alpha, Exp);
+		float const ModifiedAlpha = Pow(Alpha, Exp);
 		return Lerp<T>(A, B, ModifiedAlpha);
 	}
 
@@ -671,7 +671,7 @@ struct FMath : public FPlatformMath
 	template< class T >
 	static FORCEINLINE_DEBUGGABLE T InterpEaseOut(const T& A, const T& B, float Alpha, float Exp)
 	{
-		float const ModifiedAlpha = 1.f - FMath::Pow(1.f - Alpha, Exp);
+		float const ModifiedAlpha = 1.f - Pow(1.f - Alpha, Exp);
 		return Lerp<T>(A, B, ModifiedAlpha);
 	}
 
@@ -679,21 +679,23 @@ struct FMath : public FPlatformMath
 	template< class T > 
 	static FORCEINLINE_DEBUGGABLE T InterpEaseInOut( const T& A, const T& B, float Alpha, float Exp )
 	{
-		float const ModifiedAlpha = (Alpha < 0.5f) ?
-			0.5f * Pow(2.f * Alpha, Exp) :
-			1.f - 0.5f * Pow(2.f * (1.f - Alpha), Exp);
-
-		return Lerp<T>(A, B, ModifiedAlpha);
+		return Lerp<T>(A, B, (Alpha < 0.5f) ?
+			InterpEaseIn(0.f, 1.f, Alpha * 2.f, Exp) * 0.5f :
+			InterpEaseOut(0.f, 1.f, Alpha * 2.f - 1.f, Exp) * 0.5f + 0.5f);
 	}
 
 	/** Interpolation between A and B, applying a step function. */
 	template< class T >
 	static FORCEINLINE_DEBUGGABLE T InterpStep(const T& A, const T& B, float Alpha, int32 Steps)
 	{
-		if (Steps <= 1)
+		if (Steps <= 1 || Alpha <= 0)
 		{
 			return A;
 		}
+		else if (Alpha >= 1)
+		{
+			return B;
+		}
 
 		const float StepsAsFloat = static_cast<float>(Steps);
 		const float NumIntervals = StepsAsFloat - 1.f;
@@ -705,7 +707,7 @@ struct FMath : public FPlatformMath
 	template< class T >
 	static FORCEINLINE_DEBUGGABLE T InterpSinIn(const T& A, const T& B, float Alpha)
 	{
-		float const ModifiedAlpha = -1.f * cos(Alpha * HALF_PI) + 1.f;
+		float const ModifiedAlpha = -1.f * Cos(Alpha * HALF_PI) + 1.f;
 		return Lerp<T>(A, B, ModifiedAlpha);
 	}
 	
@@ -713,7 +715,7 @@ struct FMath : public FPlatformMath
 	template< class T >
 	static FORCEINLINE_DEBUGGABLE T InterpSinOut(const T& A, const T& B, float Alpha)
 	{
-		float const ModifiedAlpha = sin(Alpha * HALF_PI);
+		float const ModifiedAlpha = Sin(Alpha * HALF_PI);
 		return Lerp<T>(A, B, ModifiedAlpha);
 	}
 
@@ -721,16 +723,16 @@ struct FMath : public FPlatformMath
 	template< class T >
 	static FORCEINLINE_DEBUGGABLE T InterpSinInOut(const T& A, const T& B, float Alpha)
 	{
-		return (Alpha < 0.5f) ?
-			InterpSinIn(A, B, Alpha * 2.f) * 0.5f :
-			InterpSinOut(A, B, Alpha * 2.f - 1.f) * 0.5f + 0.5f;
+		return Lerp<T>(A, B, (Alpha < 0.5f) ?
+			InterpSinIn(0.f, 1.f, Alpha * 2.f) * 0.5f :
+			InterpSinOut(0.f, 1.f, Alpha * 2.f - 1.f) * 0.5f + 0.5f);
 	}
 
 	/** Interpolation between A and B, applying an exponential in function. */
 	template< class T >
 	static FORCEINLINE_DEBUGGABLE T InterpExpoIn(const T& A, const T& B, float Alpha)
 	{
-		float const ModifiedAlpha = (Alpha == 0.f) ? 0.f : pow(2.f, 10.f * (Alpha - 1.f));
+		float const ModifiedAlpha = (Alpha == 0.f) ? 0.f : Pow(2.f, 10.f * (Alpha - 1.f));
 		return Lerp<T>(A, B, ModifiedAlpha);
 	}
 
@@ -738,7 +740,7 @@ struct FMath : public FPlatformMath
 	template< class T >
 	static FORCEINLINE_DEBUGGABLE T InterpExpoOut(const T& A, const T& B, float Alpha)
 	{
-		float const ModifiedAlpha = (Alpha == 1.f) ? 1.f : -pow(2.f, -10.f * Alpha) + 1.f;
+		float const ModifiedAlpha = (Alpha == 1.f) ? 1.f : -Pow(2.f, -10.f * Alpha) + 1.f;
 		return Lerp<T>(A, B, ModifiedAlpha);
 	}
 
@@ -746,9 +748,9 @@ struct FMath : public FPlatformMath
 	template< class T >
 	static FORCEINLINE_DEBUGGABLE T InterpExpoInOut(const T& A, const T& B, float Alpha)
 	{
-		return (Alpha < 0.5f) ?
-			InterpExpoIn(A, B, Alpha * 2.f) * 0.5f :
-			InterpExpoOut(A, B, Alpha * 2.f - 1.f) * 0.5f + 0.5f;
+		return Lerp<T>(A, B, (Alpha < 0.5f) ?
+			InterpExpoIn(0.f, 1.f, Alpha * 2.f) * 0.5f :
+			InterpExpoOut(0.f, 1.f, Alpha * 2.f - 1.f) * 0.5f + 0.5f);
 	}
 
 	/** Interpolation between A and B, applying a circular in function. */
@@ -772,9 +774,9 @@ struct FMath : public FPlatformMath
 	template< class T >
 	static FORCEINLINE_DEBUGGABLE T InterpCircularInOut(const T& A, const T& B, float Alpha)
 	{
-		return (Alpha < 0.5f) ?
-			InterpCircularIn(A, B, Alpha * 2.f) * 0.5f :
-			InterpCircularOut(A, B, Alpha * 2.f - 1.f) * 0.5f + 0.5f;
+		return Lerp<T>(A, B, (Alpha < 0.5f) ?
+			InterpCircularIn(0.f, 1.f, Alpha * 2.f) * 0.5f :
+			InterpCircularOut(0.f, 1.f, Alpha * 2.f - 1.f) * 0.5f + 0.5f);
 	}
 
 	// Rotator specific interpolation
diff --git a/Engine/Source/Runtime/Core/Public/UObject/ObjectVersion.h b/Engine/Source/Runtime/Core/Public/UObject/ObjectVersion.h
--- a/Engine/Source/Runtime/Core/Public/UObject/ObjectVersion.h
+++ b/Engine/Source/Runtime/Core/Public/UObject/ObjectVersion.h
@@ -592,6 +592,10 @@ enum EUnrealEngineObjectUE4Version
 	VER_UE4_INSTANCED_STEREO_UNIFORM_UPDATE,
 	// Texture streaming min and max distance to handle HLOD
 	VER_UE4_STREAMABLE_TEXTURE_MIN_MAX_DISTANCE,	// -----<new versions can be added before this line>-------------------------------------------------
+	// Fixing up invalid struct-to-struct pin connections by injecting available conversion nodes
+	VER_UE4_INJECT_BLUEPRINT_STRUCT_PIN_CONVERSION_NODES,
+	// Saving tag data for Array Property's inner property
+	VER_UE4_INNER_ARRAY_TAG_INFO,
 	// -----<new versions can be added before this line>-------------------------------------------------
 	// - this needs to be the last line (see note below)
 	VER_UE4_AUTOMATIC_VERSION_PLUS_ONE,
diff --git a/Engine/Source/Runtime/CoreUObject/Classes/Object.h b/Engine/Source/Runtime/CoreUObject/Classes/Object.h
--- a/Engine/Source/Runtime/CoreUObject/Classes/Object.h
+++ b/Engine/Source/Runtime/CoreUObject/Classes/Object.h
@@ -216,7 +216,7 @@ namespace EMouseCursor
 
 /** A set of numerical unit types supported by the engine. Mirrored from UnitConversion.h */
 UENUM(BlueprintType)
-enum class EUnit
+enum class EUnit : uint8
 {
 	/** Scalar distance/length unit. */
 
diff --git a/Engine/Source/Runtime/CoreUObject/Private/Blueprint/BlueprintSupport.cpp b/Engine/Source/Runtime/CoreUObject/Private/Blueprint/BlueprintSupport.cpp
--- a/Engine/Source/Runtime/CoreUObject/Private/Blueprint/BlueprintSupport.cpp
+++ b/Engine/Source/Runtime/CoreUObject/Private/Blueprint/BlueprintSupport.cpp
@@ -1747,126 +1747,21 @@ void FConvertedBlueprintsDependencies::GetAssets(FName PackageName, TArray<FBlue
 	}
 }
 
-/*******************************************************************************
- * FScriptCookReplacementCoordinator
- ******************************************************************************/
 #if WITH_EDITOR
 
-FScriptCookReplacementCoordinator* FScriptCookReplacementCoordinator::CoordinatorInstance = nullptr;
+/*******************************************************************************
+* IBlueprintNativeCodeGenCore
+******************************************************************************/
+static const IBlueprintNativeCodeGenCore* CoordinatorInstance = nullptr;
 
-FScriptCookReplacementCoordinator* FScriptCookReplacementCoordinator::Get()
+const IBlueprintNativeCodeGenCore* IBlueprintNativeCodeGenCore::Get()
 {
 	return CoordinatorInstance;
 }
 
-void FScriptCookReplacementCoordinator::Create(bool bEnabled, const TArray<FString>& ExcludedAssetTypes, const TArray<FString>& ExcludedBlueprintTypes, const TMap<UObject*, UClass*>& ReplacementMap)
-{
-	CoordinatorInstance = new FScriptCookReplacementCoordinator(bEnabled, ExcludedAssetTypes, ExcludedBlueprintTypes, ReplacementMap);
-}
-
-FScriptCookReplacementCoordinator::FScriptCookReplacementCoordinator(bool bInEnabled, const TArray<FString>& InExcludedAssetTypes, const TArray<FString>& InExcludedBlueprintTypes, const TMap<UObject*, UClass*>& InReplacementMap)
-	: bEnabled(bInEnabled)
-	, ExcludedAssetTypes(InExcludedAssetTypes)
-	, ExcludedBlueprintTypes(InExcludedBlueprintTypes)
-	, ReplacementMap(InReplacementMap)
+void IBlueprintNativeCodeGenCore::Register(const IBlueprintNativeCodeGenCore* Coordinator)
 {
-}
-
-bool FScriptCookReplacementCoordinator::Initialize()
-{
-	return true;
-}
-
-DEFINE_LOG_CATEGORY_STATIC(LogDan, Log, All);
-
-UClass* FScriptCookReplacementCoordinator::FindReplacedClass(const UObject* Obj) const
-{
-	if (!bEnabled)
-	{
-		return nullptr;
-	}
-
-	// we're only looking to replace class types:
-	const UClass* AsClass = Obj->GetClass();
-	while (AsClass)
-	{
-		if (UClass* const* ReplacementPtr = ReplacementMap.Find(AsClass))
-		{
-			return *ReplacementPtr;
-		}
-
-		AsClass = AsClass->GetSuperClass();
-	}
-
-	return nullptr;
-}
-
-EReplacementResult FScriptCookReplacementCoordinator::IsTargetedForReplacement(const UPackage* Package) const
-{
-	if (!bEnabled)
-	{
-		return EReplacementResult::DontReplace;
-	}
-
-	// non-native packages with enums and structs should be converted, unless they are blacklisted:
-	UStruct* Struct = nullptr;
-	UEnum* Enum = nullptr;
-	TArray<UObject*> Objects;
-	GetObjectsWithOuter(Package, Objects, false);
-	for (auto Entry : Objects)
-	{
-		Struct = Cast<UStruct>(Entry);
-		Enum = Cast<UEnum>(Entry);
-		if (Struct || Enum)
-		{
-			break;
-		}
-	}
-
-	UObject* Target = Struct;
-	if (Target == nullptr)
-	{
-		Target = Enum;
-	}
-	return IsTargetedForReplacement(Target);
-}
-
-EReplacementResult FScriptCookReplacementCoordinator::IsTargetedForReplacement(const UObject* Object) const
-{
-	const UStruct* Struct = Cast<UStruct>(Object);
-	const UEnum* Enum = Cast<UEnum>(Object);
-
-	if (Struct == nullptr && Enum == nullptr)
-	{
-		return EReplacementResult::DontReplace;
-	}
-
-	// check blacklists:
-	// we can't use FindObject, because we may be converting a type while saving
-	if ((Struct && ExcludedAssetTypes.Find(Struct->GetPathName()) != INDEX_NONE) ||
-		(Enum && ExcludedAssetTypes.Find(Enum->GetPathName()) != INDEX_NONE))
-	{
-		return EReplacementResult::GenerateStub;
-	}
-
-	EReplacementResult Result = EReplacementResult::ReplaceCompletely;
-	while (Struct)
-	{
-		// This happens because the cooker incorrectly cooks editor only packages. Specifically happens for the blackjack sample
-		// project due to a FStringAssetReference in BaseEditor.ini:
-		if (Struct->RootPackageHasAnyFlags(PKG_EditorOnly))
-		{
-			return EReplacementResult::DontReplace;
-		}
-
-		if (ExcludedBlueprintTypes.Find(Struct->GetPathName()) != INDEX_NONE)
-		{
-			Result = EReplacementResult::GenerateStub;
-		}
-		Struct = Struct->GetSuperStruct();
-	}
-
-	return Result;
+	CoordinatorInstance = Coordinator;
 }
 
 #endif // WITH_EDITOR
diff --git a/Engine/Source/Runtime/CoreUObject/Private/UObject/Class.cpp b/Engine/Source/Runtime/CoreUObject/Private/UObject/Class.cpp
--- a/Engine/Source/Runtime/CoreUObject/Private/UObject/Class.cpp
+++ b/Engine/Source/Runtime/CoreUObject/Private/UObject/Class.cpp
@@ -1352,7 +1352,10 @@ void UStruct::SerializeTaggedProperties(FArchive& Ar, uint8* Data, UStruct* Defa
 				Tag.SerializeTaggedProperty(Ar, Property, DestAddress, DefaultsFromParent);
 
 				AdvanceProperty = true;
-				continue;
+				if (!Ar.IsCriticalError())
+				{
+					continue;
+				}
 			}
 
 			AdvanceProperty = false;
@@ -2441,6 +2444,13 @@ void UScriptStruct::InitializeStruct(void* InDest, int32 ArrayDim) const
 	}
 }
 
+#if WITH_EDITOR
+void UScriptStruct::InitializeDefaultValue(uint8* InStructData) const
+{
+	InitializeStruct(InStructData);
+}
+#endif // WITH_EDITOR
+
 void UScriptStruct::ClearScriptStruct(void* Dest, int32 ArrayDim) const
 {
 	uint8 *Data = (uint8*)Dest;
diff --git a/Engine/Source/Runtime/CoreUObject/Private/UObject/LinkerLoad.cpp b/Engine/Source/Runtime/CoreUObject/Private/UObject/LinkerLoad.cpp
--- a/Engine/Source/Runtime/CoreUObject/Private/UObject/LinkerLoad.cpp
+++ b/Engine/Source/Runtime/CoreUObject/Private/UObject/LinkerLoad.cpp
@@ -2304,17 +2304,41 @@ FLinkerLoad::EVerifyResult FLinkerLoad::VerifyImport(int32 ImportIndex)
 			// otherwise just printout warnings, and if in the editor, popup the EdLoadWarnings box
 			else
 			{
+#if UE_BUILD_DEBUG
+				bool bSkipClassWarningCheck = false;
+#endif
 				// try to get a pointer to the class of the original object so that we can display the class name of the missing resource
 				UObject* ClassPackage = FindObject<UPackage>( NULL, *Import.ClassPackage.ToString() );
 				UClass* FindClass = ClassPackage ? FindObject<UClass>( ClassPackage, *OriginalImport.ClassName.ToString() ) : NULL;
 				if( GIsEditor && !IsRunningCommandlet())
 				{
+					static const FName NAME_BlueprintGeneratedClass(TEXT("BlueprintGeneratedClass"));
+					static const UClass* ActorComponentClass = FindObjectChecked<UClass>(ANY_PACKAGE, TEXT("ActorComponent"), true);
+
 					FDeferredMessageLog LoadErrors(NAME_LoadErrors);
 					// put something into the load warnings dialog, with any extra information from above (in WarningAppend)
 					TSharedRef<FTokenizedMessage> TokenizedMessage = LoadErrors.Error(FText());
 					TokenizedMessage->AddToken(FAssetNameToken::Create(LinkerRoot->GetName()));
-					TokenizedMessage->AddToken(FTextToken::Create(FText::Format(LOCTEXT("ImportFailure", " : Failed import for {ImportClass}"), FText::FromName(GetImportClassName(ImportIndex)))));
-					TokenizedMessage->AddToken(FAssetNameToken::Create(GetImportPathName(ImportIndex)));					
+
+					if (Import.OuterIndex.IsImport() && GetImportClassName(Import.OuterIndex) == NAME_BlueprintGeneratedClass && FindClass && FindClass->IsChildOf(ActorComponentClass))
+					{
+						static const FString BPComponentArchetypePostfix(TEXT("_GEN_VARIABLE"));
+#if UE_BUILD_DEBUG
+						bSkipClassWarningCheck = true;
+#endif
+						FString ComponentVariableNameString = Import.ObjectName.ToString();
+						ComponentVariableNameString.RemoveFromEnd(BPComponentArchetypePostfix);
+						TokenizedMessage->AddToken(FTextToken::Create(FText::Format(LOCTEXT("ImportFailure", " : Failed to import {ImportClass} '{ComponentVariableName}' from"),
+							FText::FromName(GetImportClassName(ImportIndex)),
+							FText::FromString(ComponentVariableNameString))));
+						TokenizedMessage->AddToken(FAssetNameToken::Create(GetImportPathName(Import.OuterIndex.ToImport())));
+						TokenizedMessage->AddToken(FTextToken::Create(LOCTEXT("ImportFailureComponentEpilogue", "- it may have been removed or renamed. Re-saving the level may resolve this issue.")));
+					}
+					else
+					{
+						TokenizedMessage->AddToken(FTextToken::Create(FText::Format(LOCTEXT("ImportFailure", " : Failed import for {ImportClass}"), FText::FromName(GetImportClassName(ImportIndex)))));
+						TokenizedMessage->AddToken(FAssetNameToken::Create(GetImportPathName(ImportIndex)));	
+					}				
 
 					if (!WarningAppend.IsEmpty())
 					{
@@ -2326,7 +2350,7 @@ FLinkerLoad::EVerifyResult FLinkerLoad::VerifyImport(int32 ImportIndex)
 				}
 
 #if UE_BUILD_DEBUG
-				if( !IgnoreMissingReferencedClass( Import.ObjectName ) )
+				if( !bSkipClassWarningCheck && !IgnoreMissingReferencedClass( Import.ObjectName ) )
 				{
 					FUObjectThreadContext& ThreadContext = FUObjectThreadContext::Get();
 					// failure to load a class, most likely deleted instead of deprecated
@@ -4352,7 +4376,11 @@ FArchive& FLinkerLoad::operator<<( FName& Name )
 
 	if( !NameMap.IsValidIndex(NameIndex) )
 	{
-		UE_LOG(LogLinker, Fatal, TEXT("Bad name index %i/%i"), NameIndex, NameMap.Num() );
+		UE_LOG(LogLinker, Error, TEXT("Bad name index %i/%i"), NameIndex, NameMap.Num() );
+		ArIsError = true;
+		ArIsCriticalError = true;
+		Name = NAME_None;
+		return *this;
 	}
 
 	// if the name wasn't loaded (because it wasn't valid in this context)
diff --git a/Engine/Source/Runtime/CoreUObject/Private/UObject/Property.cpp b/Engine/Source/Runtime/CoreUObject/Private/UObject/Property.cpp
--- a/Engine/Source/Runtime/CoreUObject/Private/UObject/Property.cpp
+++ b/Engine/Source/Runtime/CoreUObject/Private/UObject/Property.cpp
@@ -452,7 +452,7 @@ FString UProperty::GetCPPMacroType( FString& ExtendedTypeText ) const
 }
 
 void UProperty::ExportCppDeclaration(FOutputDevice& Out, EExportedDeclaration::Type DeclarationType, const TCHAR* ArrayDimOverride, uint32 AdditionalExportCPPFlags
-	, bool bSkipParameterName, const FString* ActualCppType, const FString* ActualExtendedType) const
+	, bool bSkipParameterName, const FString* ActualCppType, const FString* ActualExtendedType, const FString* ActualParameterName) const
 {
 	const bool bIsParameter = (DeclarationType == EExportedDeclaration::Parameter) || (DeclarationType == EExportedDeclaration::MacroParameter);
 	const bool bIsInterfaceProp = dynamic_cast<const UInterfaceProperty*>(this) != nullptr;
@@ -501,14 +501,8 @@ void UProperty::ExportCppDeclaration(FOutputDevice& Out, EExportedDeclaration::T
 	FString NameCpp;
 	if (!bSkipParameterName)
 	{
-		if (AdditionalExportCPPFlags & CPPF_BlueprintCppBackend)
-		{
-			NameCpp = UnicodeToCPPIdentifier(GetName(), HasAnyPropertyFlags(CPF_Deprecated), HasAnyPropertyFlags(CPF_Parm) ? TEXT("bpp__") : TEXT("bpv__"));
-		}
-		else
-		{
-			NameCpp = GetNameCPP();
-		}
+		ensure((0 == (AdditionalExportCPPFlags & CPPF_BlueprintCppBackend)) || ActualParameterName);
+		NameCpp = ActualParameterName ? *ActualParameterName : GetNameCPP();
 	}
 	if (DeclarationType == EExportedDeclaration::MacroParameter)
 	{
diff --git a/Engine/Source/Runtime/CoreUObject/Private/UObject/PropertyArray.cpp b/Engine/Source/Runtime/CoreUObject/Private/UObject/PropertyArray.cpp
--- a/Engine/Source/Runtime/CoreUObject/Private/UObject/PropertyArray.cpp
+++ b/Engine/Source/Runtime/CoreUObject/Private/UObject/PropertyArray.cpp
@@ -2,6 +2,7 @@
 
 #include "CoreUObjectPrivate.h"
 #include "PropertyHelper.h"
+#include "PropertyTag.h"
 
 /*-----------------------------------------------------------------------------
 	UArrayProperty.
@@ -63,10 +64,78 @@ void UArrayProperty::SerializeItem( FArchive& Ar, void* Value, void const* Defau
 	}
 	ArrayHelper.CountBytes( Ar );
 
+	// Serialize a PropertyTag for the inner property of this array, allows us to validate the inner struct to see if it has changed
+	FPropertyTag InnerTag(Ar, Inner, 0, (uint8*)Value, (uint8*)Defaults);
+	if (Ar.UE4Ver() >= VER_UE4_INNER_ARRAY_TAG_INFO && InnerTag.Type == NAME_StructProperty)
+	{
+		if (Ar.IsSaving())
+		{
+			Ar << InnerTag;
+		}
+		else if (Ar.IsLoading())
+		{
+			Ar << InnerTag;
+
+			auto CanSerializeFromStructWithDifferentName = [](const FArchive& InAr, const FPropertyTag& PropertyTag, const UStructProperty* StructProperty)
+			{
+				return PropertyTag.StructGuid.IsValid() && StructProperty && StructProperty->Struct && (PropertyTag.StructGuid == StructProperty->Struct->GetCustomGuid() && (StructProperty->Struct->StructFlags & STRUCT_SerializeFromMismatchedTag));
+			};
+
+			// Check if the Inner property can successfully serialize, the type may have changed
+			UStructProperty* StructProperty = CastChecked<UStructProperty>(Inner);
+			if (InnerTag.StructName != StructProperty->Struct->GetFName()
+				&& !CanSerializeFromStructWithDifferentName(Ar, InnerTag, StructProperty))
+			{
+				UE_LOG(LogClass, Warning, TEXT("Property %s of %s has a struct type mismatch (tag %s != prop %s) in package:  %s. If that struct got renamed, add an entry to ActiveStructRedirects."),
+					*InnerTag.Name.ToString(), *GetName(), *InnerTag.StructName.ToString(), *CastChecked<UStructProperty>(Inner)->Struct->GetName(), *Ar.GetArchiveName());
+
+#if WITH_EDITOR
+				// Ensure the structure is initialized
+				for (int32 i = 0; i < n; i++)
+				{
+					StructProperty->Struct->InitializeDefaultValue(ArrayHelper.GetRawPtr(i));
+				}
+#endif // WITH_EDITOR
+
+				// Skip the property
+				const int64 StartOfProperty = Ar.Tell();
+				const int64 RemainingSize = InnerTag.Size - (Ar.Tell() - StartOfProperty);
+				uint8 B;
+				for (int64 i = 0; i < RemainingSize; i++)
+				{
+					Ar << B;
+				}
+				return;
+			}
+		}
+	}
+
+	// need to know how much data this call to SerializeItem consumes, so mark where we are
+	int32 DataOffset = Ar.Tell();
+
 	for( int32 i=0; i<n; i++ )
 	{
 		Inner->SerializeItem( Ar, ArrayHelper.GetRawPtr(i) );
 	}
+
+	if (Ar.UE4Ver() >= VER_UE4_INNER_ARRAY_TAG_INFO && Ar.IsSaving() && InnerTag.Type == NAME_StructProperty)
+	{
+		// set the tag's size
+		InnerTag.Size = Ar.Tell() - DataOffset;
+
+		if (InnerTag.Size > 0)
+		{
+			// mark our current location
+			DataOffset = Ar.Tell();
+
+			// go back and re-serialize the size now that we know it
+			Ar.Seek(InnerTag.SizeOffset);
+			Ar << InnerTag.Size;
+
+			// return to the current location
+			Ar.Seek(DataOffset);
+		}
+	}
 }
 
 bool UArrayProperty::NetSerializeItem( FArchive& Ar, UPackageMap* Map, void* Data, TArray<uint8> * MetaData ) const
@@ -284,7 +353,7 @@ void UArrayProperty::CopyValuesInternal( void* Dest, void const* Src, int32 Coun
 	int32 Num = SrcArrayHelper.Num();
 	if ( !(Inner->PropertyFlags & CPF_IsPlainOldData) )
 	{
-	DestArrayHelper.EmptyAndAddValues(Num);
+		DestArrayHelper.EmptyAndAddValues(Num);
 	}
 	else
 	{
diff --git a/Engine/Source/Runtime/CoreUObject/Private/UObject/PropertyDelegate.cpp b/Engine/Source/Runtime/CoreUObject/Private/UObject/PropertyDelegate.cpp
--- a/Engine/Source/Runtime/CoreUObject/Private/UObject/PropertyDelegate.cpp
+++ b/Engine/Source/Runtime/CoreUObject/Private/UObject/PropertyDelegate.cpp
@@ -90,6 +90,11 @@ FString UDelegateProperty::GetCPPType( FString* ExtendedTypeText/*=NULL*/, uint3
 		const bool NonNativeClassOwner = (SignatureFunction->GetOwnerClass() && !SignatureFunction->GetOwnerClass()->HasAnyClassFlags(CLASS_Native));
 		if (bBlueprintCppBackend && NonNativeClassOwner)
 		{
+			// The name must be valid, this removes spaces, ?, etc from the user's function name. It could
+			// be slightly shorter because the postfix ("__pf") is not needed here because we further post-
+			// pend to the string. Normally the postfix is needed to make sure we don't mangle to a valid
+			// identifier and collide:
+			UnmangledFunctionName = UnicodeToCPPIdentifier(UnmangledFunctionName, false, TEXT(""));
 			// the name must be unique
 			const FString OwnerName = UnicodeToCPPIdentifier(SignatureFunction->GetOwnerClass()->GetName(), false, TEXT(""));
 			const FString NewUnmangledFunctionName = FString::Printf(TEXT("%s__%s"), *UnmangledFunctionName, *OwnerName);
diff --git a/Engine/Source/Runtime/CoreUObject/Private/UObject/PropertyMulticastDelegate.cpp b/Engine/Source/Runtime/CoreUObject/Private/UObject/PropertyMulticastDelegate.cpp
--- a/Engine/Source/Runtime/CoreUObject/Private/UObject/PropertyMulticastDelegate.cpp
+++ b/Engine/Source/Runtime/CoreUObject/Private/UObject/PropertyMulticastDelegate.cpp
@@ -124,16 +124,35 @@ FString UMulticastDelegateProperty::GetCPPType( FString* ExtendedTypeText/*=NULL
 
 	FString UnmangledFunctionName = SignatureFunction->GetName().LeftChop( FString( HEADER_GENERATED_DELEGATE_SIGNATURE_SUFFIX ).Len() );
 	const UClass* OwnerClass = SignatureFunction->GetOwnerClass();
-	if ((0 != (CPPExportFlags & EPropertyExportCPPFlags::CPPF_BlueprintCppBackend)) && OwnerClass && !OwnerClass->HasAnyClassFlags(CLASS_Native))
+
+	const bool bBlueprintCppBackend = (0 != (CPPExportFlags & EPropertyExportCPPFlags::CPPF_BlueprintCppBackend));
+	const bool bNative = SignatureFunction->IsNative();
+	if (bBlueprintCppBackend && bNative)
 	{
-		// the name must be unique
-		const FString OwnerName = UnicodeToCPPIdentifier(OwnerClass->GetName(), false, TEXT(""));
-		const FString NewUnmangledFunctionName = FString::Printf(TEXT("%s__%s"), *UnmangledFunctionName, *OwnerName);
-		UnmangledFunctionName = NewUnmangledFunctionName;
+		UStruct* StructOwner = Cast<UStruct>(SignatureFunction->GetOuter());
+		if (StructOwner)
+		{
+			return FString::Printf(TEXT("%s%s::F%s"), StructOwner->GetPrefixCPP(), *StructOwner->GetName(), *UnmangledFunctionName);
+		}
 	}
-	if (0 != (CPPExportFlags & EPropertyExportCPPFlags::CPPF_CustomTypeName))
+	else
 	{
-		UnmangledFunctionName += TEXT("__MulticastDelegate");
+		if ((0 != (CPPExportFlags & EPropertyExportCPPFlags::CPPF_BlueprintCppBackend)) && OwnerClass && !OwnerClass->HasAnyClassFlags(CLASS_Native))
+		{
+			// The name must be valid, this removes spaces, ?, etc from the user's function name. It could
+			// be slightly shorter because the postfix ("__pf") is not needed here because we further post-
+			// pend to the string. Normally the postfix is needed to make sure we don't mangle to a valid
+			// identifier and collide:
+			UnmangledFunctionName = UnicodeToCPPIdentifier(UnmangledFunctionName, false, TEXT(""));
+			// the name must be unique
+			const FString OwnerName = UnicodeToCPPIdentifier(OwnerClass->GetName(), false, TEXT(""));
+			const FString NewUnmangledFunctionName = FString::Printf(TEXT("%s__%s"), *UnmangledFunctionName, *OwnerName);
+			UnmangledFunctionName = NewUnmangledFunctionName;
+		}
+		if (0 != (CPPExportFlags & EPropertyExportCPPFlags::CPPF_CustomTypeName))
+		{
+			UnmangledFunctionName += TEXT("__MulticastDelegate");
+		}
 	}
 	return FString(TEXT("F")) + UnmangledFunctionName;
 }
diff --git a/Engine/Source/Runtime/CoreUObject/Private/UObject/SavePackage.cpp b/Engine/Source/Runtime/CoreUObject/Private/UObject/SavePackage.cpp
--- a/Engine/Source/Runtime/CoreUObject/Private/UObject/SavePackage.cpp
+++ b/Engine/Source/Runtime/CoreUObject/Private/UObject/SavePackage.cpp
@@ -314,9 +314,9 @@ public:
 	 */
 	void AddReplacementsNames(UObject* Obj)
 	{
-		if (FScriptCookReplacementCoordinator* Coordinator = FScriptCookReplacementCoordinator::Get())
+		if (const IBlueprintNativeCodeGenCore* Coordinator = IBlueprintNativeCodeGenCore::Get())
 		{
-			if (const UClass* ReplObjClass = Coordinator->FindReplacedClass(Obj))
+			if (const UClass* ReplObjClass = Coordinator->FindReplacedClass(Obj->GetClass()))
 			{
 				MarkNameAsReferenced(ReplObjClass->GetFName());
 			}
@@ -363,13 +363,13 @@ public:
 };
 
 /** Returns if true if a class comes from an editor-only package */
-static bool IsEditorOnlyStruct(UStruct* InStruct)
+static bool IsEditorOnlyStruct(const UStruct* InStruct)
 {
 	// If any of the classes in this class' hierarchy comes from editor only package
 	// this class is also classified as editor-only.
-	for (UStruct* Struct = InStruct; Struct; Struct = Struct->GetSuperStruct())
+	for (const UStruct* Struct = InStruct; Struct; Struct = Struct->GetSuperStruct())
 	{
-		UPackage* StructPackage = Struct->GetOutermost();
+		const UPackage* StructPackage = Struct->GetOutermost();
 		check(StructPackage);
 		if (StructPackage->HasAnyPackageFlags(PKG_EditorOnly))
 		{
@@ -389,7 +389,7 @@ static bool IsEditorOnlyStruct(UStruct* InStruct)
  * or
  * - its outer is editor-only
 */
-static bool IsEditorOnlyObject(UObject* InObject)
+bool IsEditorOnlyObject(const UObject* InObject)
 {
 	DECLARE_SCOPE_CYCLE_COUNTER(TEXT("IsEditorOnlyObject"), STAT_IsEditorOnlyObject, STATGROUP_LoadTime);
 
@@ -418,7 +418,7 @@ static bool IsEditorOnlyObject(UObject* InObject)
 	check(InObject);
 	// If this is a package that is editor only or the object is in editor-only package,
 	// the object is editor-only too.
-	UPackage* Package = Cast<UPackage>(InObject);
+	const UPackage* Package = Cast<const UPackage>(InObject);
 	if (!Package)
 	{
 		Package = InObject->GetOutermost();
@@ -430,7 +430,7 @@ static bool IsEditorOnlyObject(UObject* InObject)
 	if (!bResult && !InObject->IsA(UPackage::StaticClass()))
 	{
 		// Otherwise the object is editor-only if its class is editor-only
-		UStruct* Struct = InObject->IsA(UStruct::StaticClass()) ? CastChecked<UStruct>(InObject) : InObject->GetClass();
+		const UStruct* Struct = InObject->IsA(UStruct::StaticClass()) ? CastChecked<const UStruct>(InObject) : InObject->GetClass();
 		bResult = IsEditorOnlyStruct(Struct);
 		if (!bResult && InObject->GetOuter())
 		{
@@ -3565,9 +3565,9 @@ ESavePackageResult UPackage::Save(UPackage* InOuter, UObject* Base, EObjectFlags
 					}
 
 #if WITH_EDITOR
-					if (FScriptCookReplacementCoordinator::Get())
+					if (IBlueprintNativeCodeGenCore::Get())
 					{
-						EReplacementResult ReplacmentResult = FScriptCookReplacementCoordinator::Get()->IsTargetedForReplacement(InOuter);
+						EReplacementResult ReplacmentResult = IBlueprintNativeCodeGenCore::Get()->IsTargetedForReplacement(InOuter);
 						if (ReplacmentResult == EReplacementResult::ReplaceCompletely)
 						{
 							return ESavePackageResult::ReplaceCompletely;
@@ -4029,9 +4029,9 @@ ESavePackageResult UPackage::Save(UPackage* InOuter, UObject* Base, EObjectFlags
 						UClass* ObjClass = nullptr;
 #if WITH_EDITOR
 						bool bReplaced = false;
-						if (FScriptCookReplacementCoordinator* Coordinator = FScriptCookReplacementCoordinator::Get())
+						if (const IBlueprintNativeCodeGenCore* Coordinator = IBlueprintNativeCodeGenCore::Get())
 						{
-							if (UClass* ReplacedClass = Coordinator->FindReplacedClass(Obj))
+							if (UClass* ReplacedClass = Coordinator->FindReplacedClass(Obj->GetClass()))
 							{
 								ObjClass = ReplacedClass;
 								bReplaced = true;
diff --git a/Engine/Source/Runtime/CoreUObject/Private/UObject/ScriptCore.cpp b/Engine/Source/Runtime/CoreUObject/Private/UObject/ScriptCore.cpp
--- a/Engine/Source/Runtime/CoreUObject/Private/UObject/ScriptCore.cpp
+++ b/Engine/Source/Runtime/CoreUObject/Private/UObject/ScriptCore.cpp
@@ -62,6 +62,7 @@ COREUOBJECT_API void GInitRunaway() {}
 // FBlueprintCoreDelegates
 
 FBlueprintCoreDelegates::FOnScriptDebuggingEvent FBlueprintCoreDelegates::OnScriptException;
+FBlueprintCoreDelegates::FOnScriptExecutionEnd FBlueprintCoreDelegates::OnScriptExecutionEnd;
 FBlueprintCoreDelegates::FOnScriptInstrumentEvent FBlueprintCoreDelegates::OnScriptProfilingEvent;
 FBlueprintCoreDelegates::FOnToggleScriptProfiler FBlueprintCoreDelegates::OnToggleScriptProfiler;
 
@@ -1160,6 +1161,9 @@ void UObject::ProcessEvent( UFunction* Function, void* Parms )
 		EScriptInstrumentationEvent EventInstrumentationInfo(EScriptInstrumentation::Stop, this);
 		FBlueprintCoreDelegates::InstrumentScriptEvent(EventInstrumentationInfo);
 	}
+#if WITH_EDITORONLY_DATA
+	FBlueprintCoreDelegates::OnScriptExecutionEnd.Broadcast();
+#endif
 #endif
 
 #if DO_BLUEPRINT_GUARD
diff --git a/Engine/Source/Runtime/CoreUObject/Private/UObject/UObjectGlobals.cpp b/Engine/Source/Runtime/CoreUObject/Private/UObject/UObjectGlobals.cpp
--- a/Engine/Source/Runtime/CoreUObject/Private/UObject/UObjectGlobals.cpp
+++ b/Engine/Source/Runtime/CoreUObject/Private/UObject/UObjectGlobals.cpp
@@ -1515,76 +1515,76 @@ void EndLoad()
  * @return	name is the form BaseName_##, where ## is the number of objects of this
  *			type that have been created since the last time the class was garbage collected.
  */
-FName MakeUniqueObjectName( UObject* Parent, UClass* Class, FName BaseName/*=NAME_None*/ )
+FName MakeUniqueObjectName( UObject* Parent, UClass* Class, FName InBaseName/*=NAME_None*/ )
 {
 	check(Class);
-	if ( BaseName == NAME_None )
-	{
-		BaseName = Class->GetFName();
-	}
+	const FName BaseName = (InBaseName == NAME_None) ? Class->GetFName() : InBaseName;
 
-	// cache the class's name's index for faster name creation later
 	FName TestName;
-	if (!FPlatformProperties::HasEditorOnlyData() && GFastPathUniqueNameGeneration)
-	{
-		/*   Fast Path Name Generation
-		* A significant fraction of object creation time goes into verifying that the a chosen unique name is really unique.
-		* The idea here is to generate unique names using very high numbers and only in situations where collisions are 
-		* impossible for other reasons.
-		*
-		* Rationale for uniqueness as used here.
-		* - Consoles do not save objects in general, and certainly not animation trees. So we could never load an object that would later clash.
-		* - We assume that we never load or create any object with a "name number" as large as, say, MAX_int32 / 2, other than via 
-		*   HACK_FastPathUniqueNameGeneration.
-		* - After using one of these large "name numbers", we decrement the static UniqueIndex, this no two names generated this way, during the
-		*   same run, could ever clash.
-		* - We assume that we could never create anywhere near MAX_int32/2 total objects at runtime, within a single run. 
-		* - We require an outer for these items, thus outers must themselves be unique. Therefore items with unique names created on the fast path
-		*   could never clash with anything with a different outer. For animation trees, these outers are never saved or loaded, thus clashes are 
-		*   impossible.
-		*/
-		static int32 UniqueIndex = MAX_int32 - 1000;
-		TestName = FName(BaseName, --UniqueIndex);
-		checkSlow(Parent); 
-		checkSlow(Parent!=ANY_PACKAGE); 
-		checkSlow(!StaticFindObjectFastInternal( NULL, Parent, TestName ));
-	}
-	else
-	{
-		UObject* ExistingObject;
-
-		do
+	do
+	{
+		// cache the class's name's index for faster name creation later
+		if (!FPlatformProperties::HasEditorOnlyData() && GFastPathUniqueNameGeneration)
+		{
+			/*   Fast Path Name Generation
+			* A significant fraction of object creation time goes into verifying that the a chosen unique name is really unique.
+			* The idea here is to generate unique names using very high numbers and only in situations where collisions are
+			* impossible for other reasons.
+			*
+			* Rationale for uniqueness as used here.
+			* - Consoles do not save objects in general, and certainly not animation trees. So we could never load an object that would later clash.
+			* - We assume that we never load or create any object with a "name number" as large as, say, MAX_int32 / 2, other than via
+			*   HACK_FastPathUniqueNameGeneration.
+			* - After using one of these large "name numbers", we decrement the static UniqueIndex, this no two names generated this way, during the
+			*   same run, could ever clash.
+			* - We assume that we could never create anywhere near MAX_int32/2 total objects at runtime, within a single run.
+			* - We require an outer for these items, thus outers must themselves be unique. Therefore items with unique names created on the fast path
+			*   could never clash with anything with a different outer. For animation trees, these outers are never saved or loaded, thus clashes are
+			*   impossible.
+			*/
+			static int32 UniqueIndex = MAX_int32 - 1000;
+			TestName = FName(BaseName, --UniqueIndex);
+			checkSlow(Parent);
+			checkSlow(Parent != ANY_PACKAGE);
+			checkSlow(!StaticFindObjectFastInternal(NULL, Parent, TestName));
+		}
+		else
 		{
-			// create the next name in the sequence for this class
-			if (BaseName.GetComparisonIndex() == NAME_Package)
+			UObject* ExistingObject;
+
+			do
 			{
-				if ( Parent == NULL )
+				// create the next name in the sequence for this class
+				if (BaseName.GetComparisonIndex() == NAME_Package)
 				{
-					//package names should default to "/Temp/Untitled" when their parent is NULL. Otherwise they are a group.
-					TestName = FName( *FString::Printf(TEXT("/Temp/%s"), *FName(NAME_Untitled).ToString()), ++Class->ClassUnique);
+					if (Parent == NULL)
+					{
+						//package names should default to "/Temp/Untitled" when their parent is NULL. Otherwise they are a group.
+						TestName = FName(*FString::Printf(TEXT("/Temp/%s"), *FName(NAME_Untitled).ToString()), ++Class->ClassUnique);
+					}
+					else
+					{
+						//package names should default to "Untitled"
+						TestName = FName(NAME_Untitled, ++Class->ClassUnique);
+					}
 				}
 				else
 				{
-					//package names should default to "Untitled"
-					TestName = FName(NAME_Untitled, ++Class->ClassUnique);
+					TestName = FName(BaseName, ++Class->ClassUnique);
 				}
-			}
-			else
-			{
-				TestName = FName(BaseName, ++Class->ClassUnique);
-			}
 
-			if (Parent == ANY_PACKAGE)
-			{
-				ExistingObject = StaticFindObject( NULL, ANY_PACKAGE, *TestName.ToString() );
-			}
-			else
-			{
-				ExistingObject = StaticFindObjectFastInternal( NULL, Parent, TestName );
-			}
-		} 
-		while( ExistingObject );
-	}
+				if (Parent == ANY_PACKAGE)
+				{
+					ExistingObject = StaticFindObject(NULL, ANY_PACKAGE, *TestName.ToString());
+				}
+				else
+				{
+					ExistingObject = StaticFindObjectFastInternal(NULL, Parent, TestName);
+				}
+			} while (ExistingObject);
+		}
+	// InBaseName can be a name of an object from a different hierarchy (so it's still unique within given parents scope), we don't want to return the same name.
+	} while (TestName == BaseName);
 	return TestName;
 }
 
diff --git a/Engine/Source/Runtime/CoreUObject/Public/Blueprint/BlueprintSupport.h b/Engine/Source/Runtime/CoreUObject/Public/Blueprint/BlueprintSupport.h
--- a/Engine/Source/Runtime/CoreUObject/Public/Blueprint/BlueprintSupport.h
+++ b/Engine/Source/Runtime/CoreUObject/Public/Blueprint/BlueprintSupport.h
@@ -79,31 +79,20 @@ enum class EReplacementResult
 };
 
 /**
- * The struct is used while saving cooked package to find replacements for 
- * converted Blueprint assets.
+ * Interface needed by CoreUObject to the BlueprintNativeCodeGen logic. Used by cooker to convert assets 
+ * to native code.
  */
-struct COREUOBJECT_API FScriptCookReplacementCoordinator
+struct IBlueprintNativeCodeGenCore
 {
-	static FScriptCookReplacementCoordinator* Get();
-	static void Create(bool bEnabled, const TArray<FString>& ExcludedAssetTypes, const TArray<FString>& ExcludedBlueprintTypes, const TMap<UObject*, UClass*>& ReplacementMap);
+	/** Returns the current IBlueprintNativeCodeGenCore, may return nullptr */
+	COREUOBJECT_API static const IBlueprintNativeCodeGenCore* Get();
 
-public:
 	/**
-	 * Enables this for handling (script) asset replacements (works only when  
-	 * running a commandlet - presumably the cooker).
-	 * 
-	 * @return True if this was successfully enabled, otherwise false.
-	 */
-	bool Initialize();
-
-	/** 
-	 * Checks to see if we're running with this on, aiming to swap out assets 
-	 * with native counterparts.
+	 * Registers the IBlueprintNativeCodeGenCore, just used to point us at an implementation.
+	 * By default, there is no IBlueprintNativeCodeGenCore, and thus no blueprints are
+	 * replaced at cook.
 	 */
-	bool IsEnabled() const
-	{
-		return bEnabled;
-	}
+	COREUOBJECT_API static void Register(const IBlueprintNativeCodeGenCore* Coordinator);
 
 	/**
 	 * Determines whether the provided package needs to be replaced (in part or completely)
@@ -111,42 +100,28 @@ public:
 	 * @param Package	The package in question
 	 * @return Whether the package should be converted
 	 */
-	EReplacementResult IsTargetedForReplacement(const UPackage* Package) const;
+	virtual EReplacementResult IsTargetedForReplacement(const UPackage* Package) const = 0;
 	
 	/**
-	* Determines whether the provided object needs to be replaced (in part or completely).
-	* Some objects in a package may require conversion and some may not. If any object 
-	* in a package wants to be converted then it is implied that all other objects will 
-	* be converted with it (no support for partial package conversion, beyond stubs)
-	*
-	* @param Object	The package in question
-	* @return Whether the object should be converted
-	*/
-	EReplacementResult IsTargetedForReplacement(const UObject* Object) const;
-
-	// Get class of converted asset. One that was specified in AddConvertedObject
-	UClass* FindReplacedClass(const UObject* Obj) const;
-
-private: 
-	/** Private so we can keep this as a singleton */
-	FScriptCookReplacementCoordinator(bool bEnabled, const TArray<FString>& ExcludedAssetTypes, const TArray<FString>& ExcludedBlueprintTypes, const TMap<UObject*, UClass*>& ReplacementMap);
+	 * Determines whether the provided object needs to be replaced (in part or completely).
+	 * Some objects in a package may require conversion and some may not. If any object 
+	 * in a package wants to be converted then it is implied that all other objects will 
+	 * be converted with it (no support for partial package conversion, beyond stubs)
+	 *
+	 * @param Object	The package in question
+	 * @return Whether the object should be converted
+	 */
+	virtual EReplacementResult IsTargetedForReplacement(const UObject* Object) const = 0;
 
 	/** 
-	 * Used to tell if replacements should be used (substituting native objects 
-	 * for assets). 
+	 * Function used to change the type of a class from, say, UBlueprintGeneratedClass to 
+	 * UDynamicClass. Cooking (and conversion in general) must be order independent so
+	 * The scope of this kind of type swap is limited.
+	 * 
+	 * @param Class The class to check for a replacement
+	 * @return A replacement class ptr, null if none
 	 */
-	bool bEnabled;
-
-	const TArray<FString> ExcludedAssetTypes;
-	const TArray<FString> ExcludedBlueprintTypes;
-
-	/** Tracks which assets has a replaced class */
-	TMap<UObject*, UClass*> ReplacementMap;
-
-	TSet<FString> ReplacedPackages;
-
-	/** singleton instance: */
-	static FScriptCookReplacementCoordinator* CoordinatorInstance;
+	virtual UClass* FindReplacedClass(const UClass* Class) const = 0;
 };
 
 #endif // WITH_EDITOR
diff --git a/Engine/Source/Runtime/CoreUObject/Public/UObject/Class.h b/Engine/Source/Runtime/CoreUObject/Public/UObject/Class.h
--- a/Engine/Source/Runtime/CoreUObject/Public/UObject/Class.h
+++ b/Engine/Source/Runtime/CoreUObject/Public/UObject/Class.h
@@ -1140,6 +1140,14 @@ public:
 	virtual COREUOBJECT_API void RecursivelyPreload();
 
 	virtual COREUOBJECT_API FGuid GetCustomGuid() const;
+
+#if WITH_EDITOR
+	/**
+	* Initializes this structure to its default values
+	* @param InStructData		The memory location to initialize
+	*/
+	virtual COREUOBJECT_API void InitializeDefaultValue(uint8* InStructData) const;
+#endif // WITH_EDITOR
 };
 
 class FStructOnScope
diff --git a/Engine/Source/Runtime/CoreUObject/Public/UObject/Script.h b/Engine/Source/Runtime/CoreUObject/Public/UObject/Script.h
--- a/Engine/Source/Runtime/CoreUObject/Public/UObject/Script.h
+++ b/Engine/Source/Runtime/CoreUObject/Public/UObject/Script.h
@@ -354,6 +354,8 @@ class COREUOBJECT_API FBlueprintCoreDelegates
 public:
 	// Callback for debugging events such as a breakpoint (Object that triggered event, active stack frame, Info)
 	DECLARE_MULTICAST_DELEGATE_ThreeParams(FOnScriptDebuggingEvent, const UObject*, const struct FFrame&, const FBlueprintExceptionInfo&);
+	// Callback for when script execution terminates.
+	DECLARE_MULTICAST_DELEGATE(FOnScriptExecutionEnd);
 	// Callback for blueprint profiling events
 	DECLARE_MULTICAST_DELEGATE_OneParam(FOnScriptInstrumentEvent, const EScriptInstrumentationEvent& );
 	// Callback for blueprint instrumentation enable/disable events
@@ -362,6 +364,8 @@ public:
 public:
 	// Called when a script exception occurs
 	static FOnScriptDebuggingEvent OnScriptException;
+	// Called when a script execution terminates.
+	static FOnScriptExecutionEnd OnScriptExecutionEnd;
 	// Called when a script profiling event is fired
 	static FOnScriptInstrumentEvent OnScriptProfilingEvent;
 	// Called when a script profiler is enabled/disabled
diff --git a/Engine/Source/Runtime/CoreUObject/Public/UObject/UObjectGlobals.h b/Engine/Source/Runtime/CoreUObject/Public/UObject/UObjectGlobals.h
--- a/Engine/Source/Runtime/CoreUObject/Public/UObject/UObjectGlobals.h
+++ b/Engine/Source/Runtime/CoreUObject/Public/UObject/UObjectGlobals.h
@@ -1876,4 +1876,18 @@ COREUOBJECT_API UPackage* FindOrConstructDynamicTypePackage(const TCHAR* Package
 /** Get names of "virtual" packages, that contain Dynamic types  */
 COREUOBJECT_API TMap<FName, FName>& GetConvertedDynamicPackageNameToTypeName();
 
+#if WITH_EDITOR
+/** 
+ * Returns if true if the object is editor-only:
+ * - it's a package marked as PKG_EditorOnly
+ * or
+ * - it's a class from a package marked as PKG_EditorOnly
+ * or
+ * - its class is from a package marked as PKG_EditorOnly
+ * or
+ * - its outer is editor-only
+ */
+COREUOBJECT_API bool IsEditorOnlyObject(const UObject* InObject);
+#endif //WITH_EDITOR
+
 #endif	// __UNOBJGLOBALS_H__
diff --git a/Engine/Source/Runtime/CoreUObject/Public/UObject/UnrealType.h b/Engine/Source/Runtime/CoreUObject/Public/UObject/UnrealType.h
--- a/Engine/Source/Runtime/CoreUObject/Public/UObject/UnrealType.h
+++ b/Engine/Source/Runtime/CoreUObject/Public/UObject/UnrealType.h
@@ -106,7 +106,7 @@ public:
 
 	// UHT interface
 	void ExportCppDeclaration(FOutputDevice& Out, EExportedDeclaration::Type DeclarationType, const TCHAR* ArrayDimOverride = NULL, uint32 AdditionalExportCPPFlags = 0
-		, bool bSkipParameterName = false, const FString* ActualCppType = nullptr, const FString* ActualExtendedType = nullptr) const;
+		, bool bSkipParameterName = false, const FString* ActualCppType = nullptr, const FString* ActualExtendedType = nullptr, const FString* ActualParameterName = nullptr) const;
 	virtual FString GetCPPMacroType( FString& ExtendedTypeText ) const;
 	virtual bool PassCPPArgsByRef() const { return false; }
 
diff --git a/Engine/Source/Runtime/Engine/Classes/EdGraph/EdGraphSchema.h b/Engine/Source/Runtime/Engine/Classes/EdGraph/EdGraphSchema.h
--- a/Engine/Source/Runtime/Engine/Classes/EdGraph/EdGraphSchema.h
+++ b/Engine/Source/Runtime/Engine/Classes/EdGraph/EdGraphSchema.h
@@ -266,38 +266,43 @@ public:
 
 public:
 	FPinConnectionResponse()
-	: Message(FText::GetEmpty())
-	, Response(CONNECT_RESPONSE_MAKE)
+		: Message(FText::GetEmpty())
+		, Response(CONNECT_RESPONSE_MAKE)
+		, bIsFatal(false)
 	{
 	}
 
 	FPinConnectionResponse(const ECanCreateConnectionResponse InResponse, const FString& InMessage)
 		: Message(FText::FromString(InMessage))
 		, Response(InResponse)
+		, bIsFatal(false)
 	{
 	}
 
 	FPinConnectionResponse(const ECanCreateConnectionResponse InResponse, const ANSICHAR* InMessage)
 		: Message(FText::FromString(InMessage))
 		, Response(InResponse)
+		, bIsFatal(false)
 	{
 	}
 
 	FPinConnectionResponse(const ECanCreateConnectionResponse InResponse, const WIDECHAR* InMessage)
 		: Message(FText::FromString(InMessage))
 		, Response(InResponse)
+		, bIsFatal(false)
 	{
 	}
 
 	FPinConnectionResponse(const ECanCreateConnectionResponse InResponse, const FText& InMessage)
 		: Message(InMessage)
 		, Response(InResponse)
+		, bIsFatal(false)
 	{
 	}
 
 	friend bool operator==(const FPinConnectionResponse& A, const FPinConnectionResponse& B)
 	{
-		return (A.Message.ToString() == B.Message.ToString()) && (A.Response == B.Response);
+		return (A.Message.ToString() == B.Message.ToString()) && (A.Response == B.Response) && (A.bIsFatal == B.bIsFatal);
 	}	
 
 	/** If a connection can be made without breaking existing connections */
@@ -305,6 +310,20 @@ public:
 	{
 		return (Response == CONNECT_RESPONSE_MAKE);
 	}
+
+	bool IsFatal() const
+	{
+		return (Response == CONNECT_RESPONSE_DISALLOW) && bIsFatal;
+	}
+
+	void SetFatal()
+	{
+		Response = CONNECT_RESPONSE_DISALLOW;
+		bIsFatal = true;
+	}
+
+private:
+	bool bIsFatal:1;
 };
 
 
diff --git a/Engine/Source/Runtime/Engine/Classes/Engine/Blueprint.h b/Engine/Source/Runtime/Engine/Classes/Engine/Blueprint.h
--- a/Engine/Source/Runtime/Engine/Classes/Engine/Blueprint.h
+++ b/Engine/Source/Runtime/Engine/Classes/Engine/Blueprint.h
@@ -498,6 +498,10 @@ public:
 	// User Defined Structures, the blueprint depends on
 	TSet<TWeakObjectPtr<UStruct>> CachedUDSDependencies;
 
+	// If this BP is just a duplicate created for a specific compilation, the reference to original GeneratedClass is needed
+	UPROPERTY(transient, duplicatetransient)
+	UClass* OriginalClass;
+
 	bool IsUpToDate() const
 	{
 		return BS_UpToDate == Status || BS_UpToDateWithWarnings == Status;
@@ -611,6 +615,9 @@ public:
 #endif	//#if WITH_EDITOR
 
 	//~ Begin UObject Interface
+#if WITH_EDITORONLY_DATA
+	virtual void PreSave() override;
+#endif // WITH_EDITORONLY_DATA
 	virtual void Serialize(FArchive& Ar) override;
 	virtual FString GetDesc(void) override;
 	virtual void TagSubobjects(EObjectFlags NewFlags) override;
diff --git a/Engine/Source/Runtime/Engine/Classes/Engine/EngineTypes.h b/Engine/Source/Runtime/Engine/Classes/Engine/EngineTypes.h
--- a/Engine/Source/Runtime/Engine/Classes/Engine/EngineTypes.h
+++ b/Engine/Source/Runtime/Engine/Classes/Engine/EngineTypes.h
@@ -3108,6 +3108,7 @@ struct FTimerHandle
 	}
 
 private:
+	UPROPERTY(transient)
 	int32 Handle;
 };
 
diff --git a/Engine/Source/Runtime/Engine/Classes/Engine/UserDefinedStruct.h b/Engine/Source/Runtime/Engine/Classes/Engine/UserDefinedStruct.h
--- a/Engine/Source/Runtime/Engine/Classes/Engine/UserDefinedStruct.h
+++ b/Engine/Source/Runtime/Engine/Classes/Engine/UserDefinedStruct.h
@@ -51,7 +51,9 @@ public:
 	virtual void PostLoad() override;
 	// End of UObject interface.
 
-	void InitializeDefaultValue(uint8* StructData) const;
+	// UScriptStruct interface.
+	virtual void InitializeDefaultValue(uint8* StructData) const override;
+	// End of UScriptStruct interface.
 
 	void ValidateGuid();
 #endif	// WITH_EDITOR
diff --git a/Engine/Source/Runtime/Engine/Classes/Kismet/KismetStringLibrary.h b/Engine/Source/Runtime/Engine/Classes/Kismet/KismetStringLibrary.h
--- a/Engine/Source/Runtime/Engine/Classes/Kismet/KismetStringLibrary.h
+++ b/Engine/Source/Runtime/Engine/Classes/Kismet/KismetStringLibrary.h
@@ -65,6 +65,21 @@ class ENGINE_API UKismetStringLibrary : public UBlueprintFunctionLibrary
 	UFUNCTION(BlueprintPure, meta=(DisplayName = "String To Float", CompactNodeTitle = "->", BlueprintAutocast), Category="Utilities|String")
 	static float Conv_StringToFloat(const FString& InString);
 	
+	/** Convert String Back To Vector. IsValid indicates whether or not the string could be successfully converted. */
+	UFUNCTION(BlueprintPure, meta=(DisplayName = "String to Vector", CompactNodeTitle = "->"), Category="Utilities|String")
+	static void Conv_StringToVector(const FString& InString, FVector& OutConvertedVector, bool& OutIsValid);
+
+	/** Convert String Back To Vector2D. IsValid indicates whether or not the string could be successfully converted. */
+	UFUNCTION(BlueprintPure, meta=(DisplayName = "String to Vector2D", CompactNodeTitle = "->"), Category="Utilities|String")
+	static void Conv_StringToVector2D(const FString& InString, FVector2D& OutConvertedVector2D, bool& OutIsValid);
+
+	/** Convert String Back To Rotator. IsValid indicates whether or not the string could be successfully converted. */
+	UFUNCTION(BlueprintPure, meta=(DisplayName = "String to Rotator", CompactNodeTitle = "->"), Category="Utilities|String")
+	static void Conv_StringToRotator(const FString& InString, FRotator& OutConvertedRotator, bool& OutIsValid);
+
+	/** Convert String Back To Color. IsValid indicates whether or not the string could be successfully converted. */
+	UFUNCTION(BlueprintPure, meta=(DisplayName = "String to Color", CompactNodeTitle = "->"), Category="Utilities|String")
+	static void Conv_StringToColor(const FString& InString, FLinearColor& OutConvertedColor, bool& OutIsValid);
 
 	/** 
 	 * Converts a float->string, create a new string in the form AppendTo+Prefix+InFloat+Suffix
diff --git a/Engine/Source/Runtime/Engine/Private/ActorConstruction.cpp b/Engine/Source/Runtime/Engine/Private/ActorConstruction.cpp
--- a/Engine/Source/Runtime/Engine/Private/ActorConstruction.cpp
+++ b/Engine/Source/Runtime/Engine/Private/ActorConstruction.cpp
@@ -610,8 +610,12 @@ UActorComponent* AActor::CreateComponentFromTemplate(UActorComponent* Template,
 			}
 		}
 
+		//Make sure, that the name of the instance is different than the name of the template. Otherwise, the template could be handled as an archetype of the instance.
+		const FName NewComponentName = (InName != NAME_None) ? InName : MakeUniqueObjectName(this, Template->GetClass(), Template->GetFName());
+		ensure(NewComponentName != Template->GetFName());
+
 		// Note we aren't copying the the RF_ArchetypeObject flag. Also note the result is non-transactional by default.
-		NewActorComp = (UActorComponent*)StaticDuplicateObject(Template, this, InName, RF_AllFlags & ~(RF_ArchetypeObject|RF_Transactional|RF_WasLoaded|RF_Public|RF_InheritableComponentTemplate) );
+		NewActorComp = (UActorComponent*)StaticDuplicateObject(Template, this, NewComponentName, RF_AllFlags & ~(RF_ArchetypeObject | RF_Transactional | RF_WasLoaded | RF_Public | RF_InheritableComponentTemplate));
 
 		NewActorComp->CreationMethod = EComponentCreationMethod::UserConstructionScript;
 
diff --git a/Engine/Source/Runtime/Engine/Private/Blueprint.cpp b/Engine/Source/Runtime/Engine/Private/Blueprint.cpp
--- a/Engine/Source/Runtime/Engine/Private/Blueprint.cpp
+++ b/Engine/Source/Runtime/Engine/Private/Blueprint.cpp
@@ -303,6 +303,16 @@ UBlueprint::UBlueprint(const FObjectInitializer& ObjectInitializer)
 {
 }
 
+#if WITH_EDITORONLY_DATA
+void UBlueprint::PreSave()
+{
+	Super::PreSave();
+
+	// Cache the BP for use
+	FFindInBlueprintSearchManager::Get().AddOrUpdateBlueprintSearchMetadata(this);
+}
+#endif // WITH_EDITORONLY_DATA
+
 void UBlueprint::Serialize(FArchive& Ar)
 {
 	Super::Serialize(Ar);
@@ -375,11 +385,6 @@ void UBlueprint::Serialize(FArchive& Ar)
 		}
 	}
 
-	if(Ar.IsSaving() && !Ar.IsTransacting())
-	{
-		// Cache the BP for use
-		FFindInBlueprintSearchManager::Get().AddOrUpdateBlueprintSearchMetadata(this);
-	}
 #endif // WITH_EDITORONLY_DATA
 }
 
diff --git a/Engine/Source/Runtime/Engine/Private/KismetMathLibrary.cpp b/Engine/Source/Runtime/Engine/Private/KismetMathLibrary.cpp
--- a/Engine/Source/Runtime/Engine/Private/KismetMathLibrary.cpp
+++ b/Engine/Source/Runtime/Engine/Private/KismetMathLibrary.cpp
@@ -4,24 +4,24 @@
 #include "Kismet/KismetMathLibrary.h"
 #include "CoreStats.h"
 
-/** Interpolate a linear alpha value using an ease mode and function, BlendExp used in cases where the easing is exponential */
+/** Interpolate a linear alpha value using an ease mode and function. */
 float EaseAlpha(float InAlpha, uint8 EasingFunc, float BlendExp, int32 Steps)
 {
 	switch (EasingFunc)
 	{
-	case EEasingFunc::Step:					return FMath::Clamp<float>(FMath::InterpStep(0.f, 1.f, InAlpha, Steps), 0.f, 1.f);
-	case EEasingFunc::SinusoidalIn:			return FMath::Clamp<float>(FMath::InterpSinIn<float>(0.f, 1.f, InAlpha), 0.f, 1.f);
-	case EEasingFunc::SinusoidalOut:		return FMath::Clamp<float>(FMath::InterpSinOut<float>(0.f, 1.f, InAlpha), 0.f, 1.f);
-	case EEasingFunc::SinusoidalInOut:		return FMath::Clamp<float>(FMath::InterpSinInOut<float>(0.f, 1.f, InAlpha), 0.f, 1.f);
-	case EEasingFunc::EaseIn:				return FMath::Clamp<float>(FMath::InterpEaseIn<float>(0.f, 1.f, InAlpha, BlendExp), 0.f, 1.f);
-	case EEasingFunc::EaseOut:				return FMath::Clamp<float>(FMath::InterpEaseOut<float>(0.f, 1.f, InAlpha, BlendExp), 0.f, 1.f);
-	case EEasingFunc::EaseInOut:			return FMath::Clamp<float>(FMath::InterpEaseInOut<float>(0.f, 1.f, InAlpha, BlendExp), 0.f, 1.f);
-	case EEasingFunc::ExpoIn:				return FMath::Clamp<float>(FMath::InterpExpoIn<float>(0.f, 1.f, InAlpha), 0.f, 1.f);
-	case EEasingFunc::ExpoOut:				return FMath::Clamp<float>(FMath::InterpExpoOut<float>(0.f, 1.f, InAlpha), 0.f, 1.f);
-	case EEasingFunc::ExpoInOut:			return FMath::Clamp<float>(FMath::InterpExpoInOut<float>(0.f, 1.f, InAlpha), 0.f, 1.f);
-	case EEasingFunc::CircularIn:			return FMath::Clamp<float>(FMath::InterpCircularIn<float>(0.f, 1.f, InAlpha), 0.f, 1.f);
-	case EEasingFunc::CircularOut:			return FMath::Clamp<float>(FMath::InterpCircularOut<float>(0.f, 1.f, InAlpha), 0.f, 1.f);
-	case EEasingFunc::CircularInOut:		return FMath::Clamp<float>(FMath::InterpCircularInOut<float>(0.f, 1.f, InAlpha), 0.f, 1.f);
+	case EEasingFunc::Step:					return FMath::InterpStep<float>(0.f, 1.f, InAlpha, Steps);
+	case EEasingFunc::SinusoidalIn:			return FMath::InterpSinIn<float>(0.f, 1.f, InAlpha);
+	case EEasingFunc::SinusoidalOut:		return FMath::InterpSinOut<float>(0.f, 1.f, InAlpha);
+	case EEasingFunc::SinusoidalInOut:		return FMath::InterpSinInOut<float>(0.f, 1.f, InAlpha);
+	case EEasingFunc::EaseIn:				return FMath::InterpEaseIn<float>(0.f, 1.f, InAlpha, BlendExp);
+	case EEasingFunc::EaseOut:				return FMath::InterpEaseOut<float>(0.f, 1.f, InAlpha, BlendExp);
+	case EEasingFunc::EaseInOut:			return FMath::InterpEaseInOut<float>(0.f, 1.f, InAlpha, BlendExp);
+	case EEasingFunc::ExpoIn:				return FMath::InterpExpoIn<float>(0.f, 1.f, InAlpha);
+	case EEasingFunc::ExpoOut:				return FMath::InterpExpoOut<float>(0.f, 1.f, InAlpha);
+	case EEasingFunc::ExpoInOut:			return FMath::InterpExpoInOut<float>(0.f, 1.f, InAlpha);
+	case EEasingFunc::CircularIn:			return FMath::InterpCircularIn<float>(0.f, 1.f, InAlpha);
+	case EEasingFunc::CircularOut:			return FMath::InterpCircularOut<float>(0.f, 1.f, InAlpha);
+	case EEasingFunc::CircularInOut:		return FMath::InterpCircularInOut<float>(0.f, 1.f, InAlpha);
 	}
 	return InAlpha;
 }
diff --git a/Engine/Source/Runtime/Engine/Private/KismetStringLibrary.cpp b/Engine/Source/Runtime/Engine/Private/KismetStringLibrary.cpp
--- a/Engine/Source/Runtime/Engine/Private/KismetStringLibrary.cpp
+++ b/Engine/Source/Runtime/Engine/Private/KismetStringLibrary.cpp
@@ -117,6 +117,26 @@ float UKismetStringLibrary::Conv_StringToFloat(const FString& InString)
 	return FCString::Atof(*InString);
 }
 
+void UKismetStringLibrary::Conv_StringToVector(const FString& InString, FVector& OutConvertedVector, bool& OutIsValid)
+{
+	OutIsValid = OutConvertedVector.InitFromString(InString);
+}
+
+void UKismetStringLibrary::Conv_StringToVector2D(const FString& InString, FVector2D& OutConvertedVector2D, bool& OutIsValid)
+{
+	OutIsValid = OutConvertedVector2D.InitFromString(InString);
+}
+
+void UKismetStringLibrary::Conv_StringToRotator(const FString& InString, FRotator& OutConvertedRotator, bool& OutIsValid)
+{
+	OutIsValid = OutConvertedRotator.InitFromString(InString);
+}
+
+void UKismetStringLibrary::Conv_StringToColor(const FString& InString, FLinearColor& OutConvertedColor, bool& OutIsValid)
+{
+	OutIsValid = OutConvertedColor.InitFromString(InString);
+}
+
 FString UKismetStringLibrary::BuildString_Float(const FString& AppendTo, const FString& Prefix, float InFloat, const FString& Suffix)
 {
 	// faster, preallocating method
diff --git a/Engine/Source/Runtime/Engine/Private/PackageMapClient.cpp b/Engine/Source/Runtime/Engine/Private/PackageMapClient.cpp
--- a/Engine/Source/Runtime/Engine/Private/PackageMapClient.cpp
+++ b/Engine/Source/Runtime/Engine/Private/PackageMapClient.cpp
@@ -1925,7 +1925,7 @@ UObject* FNetGUIDCache::GetObjectFromNetGUID( const FNetworkGUID& NetGUID, const
 
 	// At this point, we either have an outer, or we are a package
 	check( !CacheObjectPtr->bIsPending );
-	check( ObjOuter == NULL || ObjOuter->GetOutermost()->IsFullyLoaded() );
+	check(ObjOuter == NULL || ObjOuter->GetOutermost()->IsFullyLoaded() || ObjOuter->GetOutermost()->HasAnyPackageFlags(EPackageFlags::PKG_CompiledIn));
 
 	// See if this object is in memory
 	Object = StaticFindObject( UObject::StaticClass(), ObjOuter, *CacheObjectPtr->PathName.ToString(), false );
@@ -1998,7 +1998,8 @@ UObject* FNetGUIDCache::GetObjectFromNetGUID( const FNetworkGUID& NetGUID, const
 			return NULL;
 		}
 
-		if ( !Package->IsFullyLoaded() )
+		if (!Package->IsFullyLoaded() 
+			&& !Package->HasAnyPackageFlags(EPackageFlags::PKG_CompiledIn)) //TODO: dependencies of CompiledIn could still be loaded asynchronously. Are they necessary at this point??
 		{
 			if ( CVarAllowAsyncLoading.GetValueOnGameThread() > 0 )
 			{
@@ -2009,7 +2010,7 @@ UObject* FNetGUIDCache::GetObjectFromNetGUID( const FNetworkGUID& NetGUID, const
 			{
 				// If package isn't fully loaded, flush async loading now
 				FlushAsyncLoading(); 
-				check( Package->IsFullyLoaded() );
+				check(Package->IsFullyLoaded());
 			}
 		}
 
diff --git a/Engine/Source/Runtime/Engine/Public/GeneratedCodeHelpers.h b/Engine/Source/Runtime/Engine/Public/GeneratedCodeHelpers.h
--- a/Engine/Source/Runtime/Engine/Public/GeneratedCodeHelpers.h
+++ b/Engine/Source/Runtime/Engine/Public/GeneratedCodeHelpers.h
@@ -41,6 +41,11 @@ inline bool IsValid(const FScriptInterface& Test)
 	return IsValid(Test.GetObject()) && (nullptr != Test.GetInterface());
 }
 
+inline bool IsValid(const FWeakObjectPtr& Test)
+{
+	return Test.IsValid();
+}
+
 template<class TEnum>
 inline uint8 EnumToByte(TEnumAsByte<TEnum> Val)
 {
@@ -53,6 +58,12 @@ inline const T* GetDefaultValueSafe(UClass* Class)
 	return IsValid(Class) ? GetDefault<T>(Class) : nullptr;
 }
 
+template<typename ValueType>
+inline ValueType* AccessPrivateProperty(void const* ContainerPtr, int32 PropertyOffset, int32 ElementSize, int32 ArrayIndex = 0)
+{
+	return (ValueType*)((uint8*)ContainerPtr + PropertyOffset + (ElementSize * ArrayIndex));
+}
+
 struct FCustomThunkTemplates
 {
 private:
diff --git a/Engine/Source/Runtime/InputCore/Classes/InputCoreTypes.h b/Engine/Source/Runtime/InputCore/Classes/InputCoreTypes.h
--- a/Engine/Source/Runtime/InputCore/Classes/InputCoreTypes.h
+++ b/Engine/Source/Runtime/InputCore/Classes/InputCoreTypes.h
@@ -9,7 +9,7 @@ INPUTCORE_API DECLARE_LOG_CATEGORY_EXTERN(LogInput, Log, All);
 
 /** Defines the controller hands for tracking.  Could be expanded, as needed, to facilitate non-handheld controllers */
 UENUM(BlueprintType)
-enum class EControllerHand
+enum class EControllerHand : uint8
 {
 	Left,
 	Right
diff --git a/Engine/Source/Runtime/Slate/Private/Framework/Application/SlateApplication.cpp b/Engine/Source/Runtime/Slate/Private/Framework/Application/SlateApplication.cpp
--- a/Engine/Source/Runtime/Slate/Private/Framework/Application/SlateApplication.cpp
+++ b/Engine/Source/Runtime/Slate/Private/Framework/Application/SlateApplication.cpp
@@ -2164,7 +2164,7 @@ void FSlateApplication::RegisterGameViewport( TSharedRef<SViewport> InViewport )
 
 void FSlateApplication::UnregisterGameViewport()
 {
-	ResetToDefaultInputSettings();
+	ResetToDefaultPointerInputSettings();
 
 	if (GameViewportWidget.IsValid())
 	{
@@ -3021,9 +3021,12 @@ void FSlateApplication::SpawnToolTip( const TSharedRef<IToolTip>& InToolTip, con
 void FSlateApplication::CloseToolTip()
 {
 	// Notify the source widget that its tooltip is closing
-	if ( SWidget* SourceWidget = ActiveToolTipWidgetSource.Pin().Get() )
 	{
-		SourceWidget->OnToolTipClosing();
+		TSharedPtr<SWidget> SourceWidget = ActiveToolTipWidgetSource.Pin();
+		if (SourceWidget.IsValid())
+		{
+			SourceWidget->OnToolTipClosing();
+		}
 	}
 
 	// Notify the active tooltip that it's being closed.
@@ -3351,10 +3354,12 @@ void FSlateApplication::EnterDebuggingMode()
 	TSharedPtr<SViewport> PreviousGameViewport;
 
 	// Disable any game viewports while we are in debug mode so that mouse capture is released and the cursor is visible
+	// We need to retain the keyboard input for debugging purposes, so this is called directly rather than calling UnregisterGameViewport which resets input.
 	if (GameViewportWidget.IsValid())
 	{
 		PreviousGameViewport = GameViewportWidget.Pin();
-		UnregisterGameViewport();
+		PreviousGameViewport->SetActive(false);
+		GameViewportWidget.Reset();
 	}
 
 	Renderer->FlushCommands();
