Copying //UE4/Dev-Framework to //UE4/Dev-Main (Source: //UE4/Dev-Framework @ 3582324)

diff --git a/Engine/Source/Developer/BlueprintNativeCodeGen/Private/BlueprintNativeCodeGenModule.cpp b/Engine/Source/Developer/BlueprintNativeCodeGen/Private/BlueprintNativeCodeGenModule.cpp
--- a/Engine/Source/Developer/BlueprintNativeCodeGen/Private/BlueprintNativeCodeGenModule.cpp
+++ b/Engine/Source/Developer/BlueprintNativeCodeGen/Private/BlueprintNativeCodeGenModule.cpp
@@ -47,8 +47,8 @@ public:
 	virtual void FinalizeManifest() override;
 	virtual void GenerateStubs() override;
 	virtual void GenerateFullyConvertedClasses() override;
-	void MarkUnconvertedBlueprintAsNecessary(TAssetPtr<UBlueprint> BPPtr, const FCompilerNativizationOptions& NativizationOptions) override;
-	virtual const TMultiMap<FName, TAssetSubclassOf<UObject>>& GetFunctionsBoundToADelegate() override;
+	void MarkUnconvertedBlueprintAsNecessary(TSoftObjectPtr<UBlueprint> BPPtr, const FCompilerNativizationOptions& NativizationOptions) override;
+	virtual const TMultiMap<FName, TSoftClassPtr<UObject>>& GetFunctionsBoundToADelegate() override;
 
 	FFileHelper::EEncodingOptions::Type ForcedEncoding() const
 	{
@@ -79,11 +79,13 @@ private:
 	struct FStatePerPlatform
 	{
 		// A stub-wrapper must be generated only if the BP is really accessed/required by some other generated code.
-		TSet<TAssetPtr<UBlueprint>> StubsRequiredByGeneratedCode;
+		TSet<TSoftObjectPtr<UBlueprint>> StubsRequiredByGeneratedCode;
+
+		TSet<TSoftObjectPtr<UStruct>> UDSAssetsToGenerate;
+		TSet<TSoftObjectPtr<UBlueprint>> BPAssetsToGenerate;
 
-		TSet<TAssetPtr<UBlueprint>> ToGenerate;
 		// Cached values from IsTargetedForReplacement
-		mutable TMap<FStringAssetReference, EReplacementResult> CachedIsTargetedForReplacement;
+		mutable TMap<FSoftObjectPath, EReplacementResult> CachedIsTargetedForReplacement;
 	};
 
 	TMap< FName, FStatePerPlatform > StatesPerPlatform;
@@ -93,15 +95,15 @@ private:
 	// Children of these classes won't be nativized
 	TArray<FString> ExcludedAssetTypes;
 	// Eg: +ExcludedBlueprintTypes=/Script/Engine.AnimBlueprint
-	TArray<TAssetSubclassOf<UBlueprint>> ExcludedBlueprintTypes;
+	TArray<TSoftClassPtr<UBlueprint>> ExcludedBlueprintTypes;
 	// Individually excluded assets
-	TSet<FStringAssetReference> ExcludedAssets;
+	TSet<FSoftObjectPath> ExcludedAssets;
 	// Excluded folders. It excludes only BPGCs, enums and structures are still converted.
 	TArray<FString> ExcludedFolderPaths;
 
 	TArray<FName> TargetPlatformNames;
 
-	TMultiMap<FName, TAssetSubclassOf<UObject>> FunctionsBoundToADelegate; // is a function could be bound to a delegate, then it must have UFUNCTION macro. So we cannot optimize it.
+	TMultiMap<FName, TSoftClassPtr<UObject>> FunctionsBoundToADelegate; // is a function could be bound to a delegate, then it must have UFUNCTION macro. So we cannot optimize it.
 };
 
 const FCompilerNativizationOptions& FBlueprintNativeCodeGenModule::GetNativizationOptionsForPlatform(const ITargetPlatform* Platform) const
@@ -128,8 +130,8 @@ void FBlueprintNativeCodeGenModule::ReadConfig()
 		GConfig->GetArray(TEXT("BlueprintNativizationSettings"), TEXT("ExcludedBlueprintTypes"), ExcludedBlueprintTypesPath, GEditorIni);
 		for (FString& Path : ExcludedBlueprintTypesPath)
 		{
-			TAssetSubclassOf<UBlueprint> ClassPtr;
-			ClassPtr = FStringAssetReference(Path);
+			TSoftClassPtr<UBlueprint> ClassPtr;
+			ClassPtr = FSoftObjectPath(Path);
 			ClassPtr.LoadSynchronous();
 			ExcludedBlueprintTypes.Add(ClassPtr);
 		}
@@ -139,13 +141,13 @@ void FBlueprintNativeCodeGenModule::ReadConfig()
 	GConfig->GetArray(TEXT("BlueprintNativizationSettings"), TEXT("ExcludedAssets"), ExcludedAssetPaths, GEditorIni);
 	for (FString& Path : ExcludedAssetPaths)
 	{
-		ExcludedAssets.Add(FStringAssetReference(Path));
+		ExcludedAssets.Add(FSoftObjectPath(Path));
 	}
 
 	GConfig->GetArray(TEXT("BlueprintNativizationSettings"), TEXT("ExcludedFolderPaths"), ExcludedFolderPaths, GEditorIni);
 }
 
-void FBlueprintNativeCodeGenModule::MarkUnconvertedBlueprintAsNecessary(TAssetPtr<UBlueprint> BPPtr, const FCompilerNativizationOptions& NativizationOptions)
+void FBlueprintNativeCodeGenModule::MarkUnconvertedBlueprintAsNecessary(TSoftObjectPtr<UBlueprint> BPPtr, const FCompilerNativizationOptions& NativizationOptions)
 {
 	FStatePerPlatform* StateForCurrentPlatform = StatesPerPlatform.Find(NativizationOptions.PlatformName);
 	if (ensure(StateForCurrentPlatform))
@@ -169,7 +171,7 @@ void FBlueprintNativeCodeGenModule::FillTargetedForReplacementQuery()
 	};
 	ConversionQueryDelegate.BindStatic(ShouldConvert);
 
-	auto LocalMarkUnconvertedBlueprintAsNecessary = [](TAssetPtr<UBlueprint> BPPtr, const FCompilerNativizationOptions& NativizationOptions)
+	auto LocalMarkUnconvertedBlueprintAsNecessary = [](TSoftObjectPtr<UBlueprint> BPPtr, const FCompilerNativizationOptions& NativizationOptions)
 	{
 		IBlueprintNativeCodeGenModule::Get().MarkUnconvertedBlueprintAsNecessary(BPPtr, NativizationOptions);
 	};
@@ -233,7 +235,7 @@ void FBlueprintNativeCodeGenModule::CollectBoundFunctions(UBlueprint* BP)
 	}
 }
 
-const TMultiMap<FName, TAssetSubclassOf<UObject>>& FBlueprintNativeCodeGenModule::GetFunctionsBoundToADelegate()
+const TMultiMap<FName, TSoftClassPtr<UObject>>& FBlueprintNativeCodeGenModule::GetFunctionsBoundToADelegate()
 {
 	return FunctionsBoundToADelegate;
 }
@@ -338,7 +340,17 @@ void FBlueprintNativeCodeGenModule::InitializeForRerunDebugOnly(const TArray<FPl
 		}
 		FStatePerPlatform* State = StatesPerPlatform.Find(Platform.PlatformName);
 		check(State);
-		for (TAssetPtr<UBlueprint>& BPPtr : State->ToGenerate)
+
+		for (TSoftObjectPtr<UStruct>& UDSPtr : State->UDSAssetsToGenerate)
+		{
+			UStruct* UDS = UDSPtr.LoadSynchronous();
+			if (ensure(UDS))
+			{
+				GenerateSingleAsset(UDS, Platform.PlatformName);
+			}
+		}
+
+		for (TSoftObjectPtr<UBlueprint>& BPPtr : State->BPAssetsToGenerate)
 		{
 			UBlueprint* BP = BPPtr.LoadSynchronous();
 			if (ensure(BP))
@@ -361,7 +373,17 @@ void FBlueprintNativeCodeGenModule::GenerateFullyConvertedClasses()
 	{
 		FStatePerPlatform* State = StatesPerPlatform.Find(PlatformName);
 		check(State);
-		for (TAssetPtr<UBlueprint>& BPPtr : State->ToGenerate)
+
+		for (TSoftObjectPtr<UStruct>& UDSPtr : State->UDSAssetsToGenerate)
+		{
+			UStruct* UDS = UDSPtr.LoadSynchronous();
+			if (ensure(UDS))
+			{
+				GenerateSingleAsset(UDS, PlatformName, NativizationSummary);
+			}
+		}
+
+		for (TSoftObjectPtr<UBlueprint>& BPPtr : State->BPAssetsToGenerate)
 		{
 			UBlueprint* BP = BPPtr.LoadSynchronous();
 			if (ensure(BP))
@@ -511,12 +533,12 @@ void FBlueprintNativeCodeGenModule::GenerateStubs()
 		{
 			continue;
 		}
-		TSet<TAssetPtr<UBlueprint>> AlreadyGenerated;
+		TSet<TSoftObjectPtr<UBlueprint>> AlreadyGenerated;
 		while (AlreadyGenerated.Num() < StateForCurrentPlatform->StubsRequiredByGeneratedCode.Num())
 		{
 			const int32 OldGeneratedNum = AlreadyGenerated.Num();
-			TSet<TAssetPtr<UBlueprint>> LocalCopyStubsRequiredByGeneratedCode = StateForCurrentPlatform->StubsRequiredByGeneratedCode;
-			for (TAssetPtr<UBlueprint>& BPPtr : LocalCopyStubsRequiredByGeneratedCode)
+			TSet<TSoftObjectPtr<UBlueprint>> LocalCopyStubsRequiredByGeneratedCode = StateForCurrentPlatform->StubsRequiredByGeneratedCode;
+			for (TSoftObjectPtr<UBlueprint>& BPPtr : LocalCopyStubsRequiredByGeneratedCode)
 			{
 				bool bAlreadyGenerated = false;
 				AlreadyGenerated.Add(BPPtr, &bAlreadyGenerated);
@@ -571,7 +593,7 @@ void FBlueprintNativeCodeGenModule::Convert(UPackage* Package, ESavePackageResul
 	if (CookResult == ESavePackageResult::GenerateStub)
 	{
 		// No stub is generated for structs and enums.
-		ensure(!BP || !State->ToGenerate.Contains(BP));
+		ensure(!BP || !State->BPAssetsToGenerate.Contains(BP));
 	}
 	else
 	{
@@ -580,13 +602,16 @@ void FBlueprintNativeCodeGenModule::Convert(UPackage* Package, ESavePackageResul
 		{
 			if (ensure(BP))
 			{
-				State->ToGenerate.Add(BP);
+				State->BPAssetsToGenerate.Add(BP);
 			}
 		}
+		else if (Struct)
+		{
+			State->UDSAssetsToGenerate.Add(Struct);
+		}
 		else
 		{
-			UField* ForConversion = Struct ? (UField*)Struct : (UField*)Enum;
-			GenerateSingleAsset(ForConversion, PlatformName);
+			GenerateSingleAsset((UField*)Enum, PlatformName);
 		}
 	}
 }
@@ -620,10 +645,10 @@ void FBlueprintNativeCodeGenModule::FinalizeManifest()
 		{
 			if (NativizationSummary.IsValid())
 			{
-				TSet<TAssetPtr<UPackage>>* RequiredModules = NativizationSummary->ModulesRequiredByPlatform.Find(Manifest.GetCompilerNativizationOptions().PlatformName);
+				TSet<TSoftObjectPtr<UPackage>>* RequiredModules = NativizationSummary->ModulesRequiredByPlatform.Find(Manifest.GetCompilerNativizationOptions().PlatformName);
 				if (RequiredModules)
 				{
-					for (TAssetPtr<UPackage> ItPackage : *RequiredModules)
+					for (TSoftObjectPtr<UPackage> ItPackage : *RequiredModules)
 					{
 						if (UPackage* Pkg = ItPackage.Get())
 						{
@@ -642,22 +667,34 @@ void FBlueprintNativeCodeGenModule::FinalizeManifest()
 UClass* FBlueprintNativeCodeGenModule::FindReplacedClassForObject(const UObject* Object, const FCompilerNativizationOptions& NativizationOptions) const
 {
 	// we're only looking to replace class types:
-	if (Object && Object->IsA<UField>() 
-		&& (EReplacementResult::ReplaceCompletely == IsTargetedForReplacement(Object, NativizationOptions)))
+	if (Object && Object->IsA<UField>())
 	{
-		for (const UClass* Class = Object->GetClass(); Class; Class = Class->GetSuperClass())
+		if (IsTargetedForReplacement(Object, NativizationOptions) == EReplacementResult::ReplaceCompletely)
 		{
-			if (Class == UUserDefinedEnum::StaticClass())
-			{
-				return UEnum::StaticClass();
-			}
-			if (Class == UUserDefinedStruct::StaticClass())
+			for (const UClass* Class = Object->GetClass(); Class; Class = Class->GetSuperClass())
 			{
-				return UScriptStruct::StaticClass();
+				if (Class == UUserDefinedEnum::StaticClass())
+				{
+					return UEnum::StaticClass();
+				}
+				if (Class == UUserDefinedStruct::StaticClass())
+				{
+					return UScriptStruct::StaticClass();
+				}
+				if (Class == UBlueprintGeneratedClass::StaticClass())
+				{
+					return UDynamicClass::StaticClass();
+				}
 			}
-			if (Class == UBlueprintGeneratedClass::StaticClass())
+		}
+		else if (const UByteProperty* ByteProperty = Cast<UByteProperty>(Object))
+		{
+			// User-Defined Enum values are compiled as Byte properties, but get converted to Enum class properties during nativization. Thus,
+			// we have to account for that here and switch the property class to be an Enum property, since that's what will be generated by UHT.
+			// If we don't do this, then a dependent asset's import table will contain the incorrect property class for this value, if referenced.
+			if (ByteProperty->Enum && IsTargetedForReplacement(ByteProperty->Enum, NativizationOptions) == EReplacementResult::ReplaceCompletely)
 			{
-				return UDynamicClass::StaticClass();
+				return UEnumProperty::StaticClass();
 			}
 		}
 	}
@@ -773,7 +810,7 @@ EReplacementResult FBlueprintNativeCodeGenModule::IsTargetedForReplacement(const
 
 	const FStatePerPlatform* StateForCurrentPlatform = StatesPerPlatform.Find(NativizationOptions.PlatformName);
 	check(StateForCurrentPlatform);
-	const FStringAssetReference ObjectKey(Object);
+	const FSoftObjectPath ObjectKey(Object);
 	{
 		const EReplacementResult* const CachedValue = StateForCurrentPlatform->CachedIsTargetedForReplacement.Find(ObjectKey); //THe referenced returned by FindOrAdd could be invalid later, when filled.
 		if (CachedValue)
@@ -934,7 +971,7 @@ EReplacementResult FBlueprintNativeCodeGenModule::IsTargetedForReplacement(const
 			}
 
 			// ExcludedBlueprintTypes
-			for (TAssetSubclassOf<UBlueprint> ExcludedBlueprintTypeAsset : ExcludedBlueprintTypes)
+			for (TSoftClassPtr<UBlueprint> ExcludedBlueprintTypeAsset : ExcludedBlueprintTypes)
 			{
 				UClass* ExcludedBPClass = ExcludedBlueprintTypeAsset.Get();
 				if (!ExcludedBPClass)
@@ -1036,7 +1073,7 @@ void FBlueprintNativeCodeGenModule::FillPlatformNativizationDetails(const ITarge
 			GConfig->GetArray(ConfigSection, KeyForExcludedAssets, ExcludedAssetPaths, GEditorIni);
 			for (FString& Path : ExcludedAssetPaths)
 			{
-				Details.CompilerNativizationOptions.ExcludedAssets.Add(FStringAssetReference(Path));
+				Details.CompilerNativizationOptions.ExcludedAssets.Add(FSoftObjectPath(Path));
 			}
 		}
 	};
diff --git a/Engine/Source/Developer/BlueprintNativeCodeGen/Public/BlueprintNativeCodeGenModule.h b/Engine/Source/Developer/BlueprintNativeCodeGen/Public/BlueprintNativeCodeGenModule.h
--- a/Engine/Source/Developer/BlueprintNativeCodeGen/Public/BlueprintNativeCodeGenModule.h
+++ b/Engine/Source/Developer/BlueprintNativeCodeGen/Public/BlueprintNativeCodeGenModule.h
@@ -4,7 +4,7 @@
 
 #include "CoreMinimal.h"
 #include "UObject/Object.h"
-#include "UObject/AssetPtr.h"
+#include "UObject/SoftObjectPtr.h"
 #include "Modules/ModuleInterface.h"
 #include "Modules/ModuleManager.h"
 #include "Engine/Blueprint.h"
@@ -84,8 +84,8 @@ public:
 	virtual void FinalizeManifest() = 0;
 	virtual void GenerateStubs() = 0;
 	virtual void GenerateFullyConvertedClasses() = 0;
-	virtual void MarkUnconvertedBlueprintAsNecessary(TAssetPtr<UBlueprint> BPPtr, const FCompilerNativizationOptions& NativizationOptions) = 0;
-	virtual const TMultiMap<FName, TAssetSubclassOf<UObject>>& GetFunctionsBoundToADelegate() = 0;
+	virtual void MarkUnconvertedBlueprintAsNecessary(TSoftObjectPtr<UBlueprint> BPPtr, const FCompilerNativizationOptions& NativizationOptions) = 0;
+	virtual const TMultiMap<FName, TSoftClassPtr<UObject>>& GetFunctionsBoundToADelegate() = 0;
 
 	virtual void FillPlatformNativizationDetails(const ITargetPlatform* Platform, FPlatformNativizationDetails& OutDetails) = 0;
 protected:
diff --git a/Engine/Source/Developer/FunctionalTesting/Classes/ScreenshotFunctionalTest.h b/Engine/Source/Developer/FunctionalTesting/Classes/ScreenshotFunctionalTest.h
--- a/Engine/Source/Developer/FunctionalTesting/Classes/ScreenshotFunctionalTest.h
+++ b/Engine/Source/Developer/FunctionalTesting/Classes/ScreenshotFunctionalTest.h
@@ -12,8 +12,8 @@
 /**
  * 
  */
-UCLASS(Blueprintable, MinimalAPI)
-class AScreenshotFunctionalTest : public AFunctionalTest
+UCLASS(Blueprintable)
+class FUNCTIONALTESTING_API AScreenshotFunctionalTest : public AFunctionalTest
 {
 	GENERATED_BODY()
 
diff --git a/Engine/Source/Developer/FunctionalTesting/Private/FunctionalAITest.cpp b/Engine/Source/Developer/FunctionalTesting/Private/FunctionalAITest.cpp
--- a/Engine/Source/Developer/FunctionalTesting/Private/FunctionalAITest.cpp
+++ b/Engine/Source/Developer/FunctionalTesting/Private/FunctionalAITest.cpp
@@ -3,7 +3,7 @@
 #include "FunctionalAITest.h"
 #include "TimerManager.h"
 #include "Engine/World.h"
-#include "FuncTestManager.h"
+#include "FunctionalTestingModule.h"
 #include "FunctionalTestingManager.h"
 #include "AI/Navigation/NavigationSystem.h"
 #include "AIController.h"
diff --git a/Engine/Source/Developer/FunctionalTesting/Private/FunctionalTest.cpp b/Engine/Source/Developer/FunctionalTesting/Private/FunctionalTest.cpp
--- a/Engine/Source/Developer/FunctionalTesting/Private/FunctionalTest.cpp
+++ b/Engine/Source/Developer/FunctionalTesting/Private/FunctionalTest.cpp
@@ -1,6 +1,7 @@
 // Copyright 1998-2017 Epic Games, Inc. All Rights Reserved.
 
 #include "FunctionalTest.h"
+#include "FunctionalTestingModule.h"
 #include "Misc/Paths.h"
 #include "Engine/GameViewportClient.h"
 #include "Engine/LatentActionManager.h"
@@ -13,7 +14,6 @@
 #include "GameFramework/PlayerController.h"
 #include "Components/TextRenderComponent.h"
 #include "Engine/Selection.h"
-#include "FuncTestManager.h"
 #include "FuncTestRenderingComponent.h"
 #include "ObjectEditorUtils.h"
 #include "VisualLogger/VisualLogger.h"
diff --git a/Engine/Source/Developer/FunctionalTesting/Private/FunctionalTestingManager.cpp b/Engine/Source/Developer/FunctionalTesting/Private/FunctionalTestingManager.cpp
--- a/Engine/Source/Developer/FunctionalTesting/Private/FunctionalTestingManager.cpp
+++ b/Engine/Source/Developer/FunctionalTesting/Private/FunctionalTestingManager.cpp
@@ -4,7 +4,6 @@
 #include "TimerManager.h"
 #include "Engine/World.h"
 #include "EngineUtils.h"
-#include "FuncTestManager.h"
 #include "FunctionalTestingModule.h"
 #include "AI/Navigation/NavigationSystem.h"
 #include "EngineGlobals.h"
@@ -184,14 +183,14 @@ void UFunctionalTestingManager::TriggerFirstValidTest()
 
 UFunctionalTestingManager* UFunctionalTestingManager::GetManager(UObject* WorldContext)
 {
-	UFunctionalTestingManager* Manager = FFunctionalTestingModule::Get()->GetCurrentScript();
+	UFunctionalTestingManager* Manager = IFunctionalTestingModule::Get().GetCurrentManager();
 
 	if (Manager == nullptr)
 	{
 		if (UObject* World = GEngine->GetWorldFromContextObject(WorldContext, EGetWorldErrorMode::LogAndReturnNull))
 		{
 			Manager = NewObject<UFunctionalTestingManager>(World);
-			FFunctionalTestingModule::Get()->SetScript(Manager);
+			IFunctionalTestingModule::Get().SetManager(Manager);
 
 			// add to root and get notified on world cleanup to remove from root on map cleanup
 			Manager->AddToRoot();
@@ -215,7 +214,7 @@ void UFunctionalTestingManager::OnWorldCleanedUp(UWorld* World, bool bSessionEnd
 		RemoveFromRoot();
 
 		// Clear the functional test manager once the world is removed.
-		FFunctionalTestingModule::Get()->SetScript(nullptr);
+		IFunctionalTestingModule::Get().SetManager(nullptr);
 	}
 }
 
diff --git a/Engine/Source/Developer/FunctionalTesting/Private/FunctionalTestingModule.cpp b/Engine/Source/Developer/FunctionalTesting/Private/FunctionalTestingModule.cpp
--- a/Engine/Source/Developer/FunctionalTesting/Private/FunctionalTestingModule.cpp
+++ b/Engine/Source/Developer/FunctionalTesting/Private/FunctionalTestingModule.cpp
@@ -1,20 +1,269 @@
 // Copyright 1998-2017 Epic Games, Inc. All Rights Reserved.
 
 #include "FunctionalTestingModule.h"
-#include "FuncTestManager.h"
+#include "FunctionalTestingManager.h"
+#include "Misc/CoreMisc.h"
+#include "Engine/World.h"
+#include "Engine/Engine.h"
+#include "EngineUtils.h"
+#include "FunctionalTest.h"
+#include "EngineGlobals.h"
+#include "ARFilter.h"
+#include "AssetRegistryModule.h"
 
 #define LOCTEXT_NAMESPACE "FunctionalTesting"
 
+DEFINE_LOG_CATEGORY(LogFunctionalTest);
+
+class FFunctionalTestingModule : public IFunctionalTestingModule
+{
+public:
+	virtual void StartupModule() override;
+	virtual void ShutdownModule() override;
+
+	virtual void RunAllTestsOnMap(bool bClearLog, bool bRunLooped) override;
+	virtual void RunTestOnMap(const FString& TestName, bool bClearLog, bool bRunLooped) override;
+	virtual void MarkPendingActivation() override;
+	virtual bool IsActivationPending() const override;
+	virtual bool IsRunning() const override;
+	virtual bool IsFinished() const override;
+	virtual void SetManager(class UFunctionalTestingManager* NewManager) override;
+	virtual class UFunctionalTestingManager* GetCurrentManager();
+	virtual void SetLooping(const bool bLoop) override;
+	virtual void GetMapTests(bool bEditorOnlyTests, TArray<FString>& OutBeautifiedNames, TArray<FString>& OutTestCommands, TArray<FString>& OutTestMapAssets) const override;
+
+private:
+	UWorld* GetTestWorld();
+	void OnGetAssetTagsForWorld(const UWorld* World, TArray<UObject::FAssetRegistryTag>& OutTags);
+
+	TWeakObjectPtr<class UFunctionalTestingManager> TestManager;
+	bool bPendingActivation;
+};
+
 void FFunctionalTestingModule::StartupModule() 
 {
-	Manager = MakeShareable(new FFuncTestManager());
+	bPendingActivation = false;
+#if WITH_EDITOR
+	FWorldDelegates::GetAssetTags.AddRaw(this, &FFunctionalTestingModule::OnGetAssetTagsForWorld);
+#endif
 }
 
 void FFunctionalTestingModule::ShutdownModule() 
 {
-	Manager.Reset();
+#if WITH_EDITOR
+	FWorldDelegates::GetAssetTags.RemoveAll(this);
+#endif
+}
+
+void FFunctionalTestingModule::OnGetAssetTagsForWorld(const UWorld* World, TArray<UObject::FAssetRegistryTag>& OutTags)
+{
+#if WITH_EDITOR
+	int32 Tests = 0;
+	FString TestNames, TestNamesEditor;
+	for (TActorIterator<AFunctionalTest> ActorItr(const_cast<UWorld*>(World), AFunctionalTest::StaticClass(), EActorIteratorFlags::AllActors); ActorItr; ++ActorItr)
+	{
+		AFunctionalTest* FunctionalTest = *ActorItr;
+
+		// Only include enabled tests in the list of functional tests to run.
+		if (FunctionalTest->IsEnabled())
+		{
+			bool bIsEditorOnly = IsEditorOnlyObject(FunctionalTest);
+
+			// Check if this class is editor only
+			FString& NamesAppend = bIsEditorOnly ? TestNamesEditor : TestNames;
+
+			Tests++;
+			NamesAppend.Append(FunctionalTest->GetActorLabel() + TEXT("|") + FunctionalTest->GetName());
+			NamesAppend.Append(TEXT(";"));
+		}
+	}
+
+	if (Tests > 0)
+	{
+		OutTags.Add(UObject::FAssetRegistryTag("Tests", FString::FromInt(Tests), UObject::FAssetRegistryTag::TT_Numerical));
+	}
+
+	if (!TestNames.IsEmpty())
+	{
+		OutTags.Add(UObject::FAssetRegistryTag("TestNames", TestNames, UObject::FAssetRegistryTag::TT_Hidden));
+	}
+
+	if (!TestNamesEditor.IsEmpty())
+	{
+		OutTags.Add(UObject::FAssetRegistryTag("TestNamesEditor", TestNamesEditor, UObject::FAssetRegistryTag::TT_Hidden));
+	}
+#endif
+}
+
+void FFunctionalTestingModule::GetMapTests(bool bEditorOnlyTests, TArray<FString>& OutBeautifiedNames, TArray<FString>& OutTestCommands, TArray<FString>& OutTestMapAssets) const
+{
+	IAssetRegistry& AssetRegistry = FModuleManager::Get().LoadModuleChecked<FAssetRegistryModule>("AssetRegistry").Get();
+
+	if (!AssetRegistry.IsLoadingAssets())
+	{
+		TArray<FAssetData> MapList;
+		FARFilter Filter;
+		Filter.ClassNames.Add(UWorld::StaticClass()->GetFName());
+		Filter.bRecursiveClasses = true;
+		Filter.bIncludeOnlyOnDiskAssets = true;
+		if (AssetRegistry.GetAssets(Filter, /*out*/ MapList))
+		{
+			for (const FAssetData& MapAsset : MapList)
+			{
+				FString MapAssetPath = MapAsset.ObjectPath.ToString();
+
+				const FString* Tests = MapAsset.TagsAndValues.Find(TEXT("Tests"));
+				const FString* TestNames = MapAsset.TagsAndValues.Find(bEditorOnlyTests ? TEXT("TestNamesEditor") : TEXT("TestNames"));
+
+				if (Tests && TestNames)
+				{
+					int32 TestCount = FCString::Atoi(**Tests);
+					if (TestCount > 0)
+					{
+						TArray<FString> MapTests;
+						(*TestNames).ParseIntoArray(MapTests, TEXT(";"), true);
+
+						for (const FString& MapTest : MapTests)
+						{
+							FString BeautifulTestName;
+							FString RealTestName;
+
+							if (MapTest.Split(TEXT("|"), &BeautifulTestName, &RealTestName))
+							{
+								OutBeautifiedNames.Add(MapAsset.PackageName.ToString() + TEXT(".") + *BeautifulTestName);
+								OutTestCommands.Add(MapAssetPath + TEXT(";") + MapAsset.PackageName.ToString() + TEXT(";") + *RealTestName);
+								OutTestMapAssets.AddUnique(MapAssetPath);
+							}
+						}
+					}
+				}
+				else if (!bEditorOnlyTests && MapAsset.AssetName.ToString().Find(TEXT("FTEST_")) == 0)
+				{
+					OutBeautifiedNames.Add(MapAsset.AssetName.ToString());
+					OutTestCommands.Add(MapAssetPath + TEXT(";") + MapAsset.PackageName.ToString());
+					OutTestMapAssets.AddUnique(MapAssetPath);
+				}
+			}
+		}
+	}
+}
+
+void FFunctionalTestingModule::SetManager(class UFunctionalTestingManager* NewManager)
+{
+	TestManager = NewManager;
+}
+
+UFunctionalTestingManager* FFunctionalTestingModule::GetCurrentManager()
+{
+	return TestManager.Get();
+}
+
+bool FFunctionalTestingModule::IsRunning() const
+{
+	return TestManager.IsValid() && TestManager->IsRunning();
 }
 
+bool FFunctionalTestingModule::IsFinished() const
+{
+	return (!TestManager.IsValid() || TestManager->IsFinished());
+}
+
+void FFunctionalTestingModule::MarkPendingActivation()
+{
+	bPendingActivation = true;
+}
+
+bool FFunctionalTestingModule::IsActivationPending() const
+{
+	return bPendingActivation;
+}
+
+void FFunctionalTestingModule::SetLooping(const bool bLoop)
+{
+	if (TestManager.IsValid())
+	{
+		TestManager->SetLooped(bLoop);
+	}
+}
+
+UWorld* FFunctionalTestingModule::GetTestWorld()
+{
+	UWorld* TestWorld = nullptr;
+
+#if WITH_EDITOR
+	const TIndirectArray<FWorldContext>& WorldContexts = GEngine->GetWorldContexts();
+	for (const FWorldContext& Context : WorldContexts)
+	{
+		if ((Context.WorldType == EWorldType::PIE) && (Context.World() != nullptr))
+		{
+			TestWorld = Context.World();
+		}
+	}
+#endif
+	if (!TestWorld)
+	{
+		TestWorld = GWorld;
+		if (GIsEditor)
+		{
+			UE_LOG(LogFunctionalTest, Warning, TEXT("Functional Test using GWorld.  Not correct for PIE"));
+		}
+	}
+
+	return TestWorld;
+}
+
+void FFunctionalTestingModule::RunAllTestsOnMap(bool bClearLog, bool bRunLooped)
+{
+	if (UWorld* TestWorld = GetTestWorld())
+	{
+		bPendingActivation = false;
+		if (UFunctionalTestingManager::RunAllFunctionalTests(TestWorld, bClearLog, bRunLooped) == false)
+		{
+			UE_LOG(LogFunctionalTest, Error, TEXT("No functional testing script on map."));
+		}
+	}
+}
+
+void FFunctionalTestingModule::RunTestOnMap(const FString& TestName, bool bClearLog, bool bRunLooped)
+{
+	if (UWorld* TestWorld = GetTestWorld())
+	{
+		bPendingActivation = false;
+		if (UFunctionalTestingManager::RunAllFunctionalTests(TestWorld, bClearLog, bRunLooped, true, TestName) == false)
+		{
+			UE_LOG(LogFunctionalTest, Error, TEXT("No functional testing script on map."));
+		}
+	}
+}
+
+//////////////////////////////////////////////////////////////////////////
+// Exec
+//////////////////////////////////////////////////////////////////////////
+
+static bool FuncTestExec(UWorld* InWorld, const TCHAR* Command, FOutputDevice& Ar)
+{
+	if (FParse::Command(&Command, TEXT("ftest")))
+	{
+		if (FParse::Command(&Command, TEXT("start")))
+		{
+			const bool bLooped = FParse::Command(&Command, TEXT("loop"));
+
+			//instead of allowing straight use of the functional test framework, this should go through the automation framework and kick off one of the Editor/Client functional tests
+
+			IFunctionalTestingModule& Module = IFunctionalTestingModule::Get();
+			if (!Module.IsRunning() && !Module.IsActivationPending())
+			{
+				Module.RunAllTestsOnMap(/*bClearLog=*/true, bLooped);
+			}
+		}
+		return true;
+	}
+	return false;
+}
+
+FStaticSelfRegisteringExec FuncTestExecRegistration(FuncTestExec);
+
+
 IMPLEMENT_MODULE( FFunctionalTestingModule, FunctionalTesting );
 
 #undef LOCTEXT_NAMESPACE
diff --git a/Engine/Source/Developer/FunctionalTesting/Private/Tests/ClientFuncTestPerforming.cpp b/Engine/Source/Developer/FunctionalTesting/Private/Tests/ClientFuncTestPerforming.cpp
--- a/Engine/Source/Developer/FunctionalTesting/Private/Tests/ClientFuncTestPerforming.cpp
+++ b/Engine/Source/Developer/FunctionalTesting/Private/Tests/ClientFuncTestPerforming.cpp
@@ -8,35 +8,31 @@
 #include "FunctionalTestingModule.h"
 #include "EngineGlobals.h"
 #include "Tests/AutomationCommon.h"
-#include "ARFilter.h"
-#include "AssetRegistryModule.h"
 
 #if WITH_DEV_AUTOMATION_TESTS
 
 #define LOCTEXT_NAMESPACE "FunctionalTesting"
 
-DEFINE_LOG_CATEGORY_STATIC(LogFunctionalTesting, Log, All);
-
 DEFINE_LATENT_AUTOMATION_COMMAND(FWaitForFTestsToFinish);
 bool FWaitForFTestsToFinish::Update()
 {
-	return FFunctionalTestingModule::Get()->IsRunning() == false;
+	return IFunctionalTestingModule::Get().IsRunning() == false;
 }
 
 DEFINE_LATENT_AUTOMATION_COMMAND(FTriggerFTests);
 bool FTriggerFTests::Update()
 {
-	IFuncTestManager* Manager = FFunctionalTestingModule::Get();
-	if (Manager->IsFinished())
+	IFunctionalTestingModule& Module = IFunctionalTestingModule::Get();
+	if (Module.IsFinished())
 	{
 		// if tests have been already triggered by level script just make sure it's not looping
-		if (Manager->IsRunning())
+		if (Module.IsRunning())
 		{
-			FFunctionalTestingModule::Get()->SetLooping(false);
+			Module.SetLooping(false);
 		}
 		else
 		{
-			FFunctionalTestingModule::Get()->RunAllTestsOnMap(false, false);
+			Module.RunAllTestsOnMap(false, false);
 		}
 	}
 
@@ -47,17 +43,17 @@ bool FTriggerFTests::Update()
 DEFINE_LATENT_AUTOMATION_COMMAND_ONE_PARAMETER(FTriggerFTest, FString, TestName);
 bool FTriggerFTest::Update()
 {
-	IFuncTestManager* Manager = FFunctionalTestingModule::Get();
-	if ( Manager->IsFinished() )
+	IFunctionalTestingModule& Module = IFunctionalTestingModule::Get();
+	if (Module.IsFinished() )
 	{
 		// if tests have been already triggered by level script just make sure it's not looping
-		if ( Manager->IsRunning() )
+		if (Module.IsRunning() )
 		{
-			FFunctionalTestingModule::Get()->SetLooping(false);
+			Module.SetLooping(false);
 		}
 		else
 		{
-			FFunctionalTestingModule::Get()->RunTestOnMap(TestName, false, false);
+			Module.RunTestOnMap(TestName, false, false);
 		}
 	}
 
@@ -88,24 +84,6 @@ bool FStartFTestOnMap::Update()
 	return true;
 }
 
-/**
- * 
- */
-
-// Project.Maps.Client Functional Testing
-// Project.Maps.Functional Tests
-
-void ParseTestMapInfo(const FString& Parameters, FString& MapObjectPath, FString& MapPackageName, FString& MapTestName)
-{
-	TArray<FString> ParamArray;
-	Parameters.ParseIntoArray(ParamArray, TEXT(";"), true);
-
-	MapObjectPath = ParamArray[0];
-	MapPackageName = ParamArray[1];
-	MapTestName = ( ParamArray.Num() > 2 ) ? ParamArray[2] : TEXT("");
-}
-
-
 class FClientFunctionalTestingMapsBase : public FAutomationTestBase
 {
 public:
@@ -114,6 +92,37 @@ public:
 	{
 	}
 
+	// Project.Maps.Client Functional Testing
+	// Project.Maps.Functional Tests
+
+	static void ParseTestMapInfo(const FString& Parameters, FString& MapObjectPath, FString& MapPackageName, FString& MapTestName)
+	{
+		TArray<FString> ParamArray;
+		Parameters.ParseIntoArray(ParamArray, TEXT(";"), true);
+
+		MapObjectPath = ParamArray[0];
+		MapPackageName = ParamArray[1];
+		MapTestName = (ParamArray.Num() > 2) ? ParamArray[2] : TEXT("");
+	}
+
+	// @todo this is a temporary solution. Once we know how to get test's hands on a proper world
+	// this function should be redone/removed
+	static UWorld* GetAnyGameWorld()
+	{
+		UWorld* TestWorld = nullptr;
+		const TIndirectArray<FWorldContext>& WorldContexts = GEngine->GetWorldContexts();
+		for (const FWorldContext& Context : WorldContexts)
+		{
+			if (((Context.WorldType == EWorldType::PIE) || (Context.WorldType == EWorldType::Game)) && (Context.World() != NULL))
+			{
+				TestWorld = Context.World();
+				break;
+			}
+		}
+
+		return TestWorld;
+	}
+
 	virtual FString GetTestOpenCommand(const FString& Parameters) const override
 	{
 		FString MapObjectPath, MapPackageName, MapTestName;
@@ -129,132 +138,91 @@ public:
 
 		return MapObjectPath;
 	}
-};
 
+	/** 
+	 * Requests a enumeration of all maps to be loaded
+	 */
+	virtual void GetTests(TArray<FString>& OutBeautifiedNames, TArray<FString>& OutTestCommands) const override
+	{
+		bool bEditorOnlyTests = !(GetTestFlags() & EAutomationTestFlags::ClientContext);
+		TArray<FString> MapAssets;
+		IFunctionalTestingModule::Get().GetMapTests(bEditorOnlyTests, OutBeautifiedNames, OutTestCommands, MapAssets);
+	}
 
-// create test base class
-IMPLEMENT_CUSTOM_COMPLEX_AUTOMATION_TEST(FClientFunctionalTestingMaps, FClientFunctionalTestingMapsBase, "Project.Functional Tests", (EAutomationTestFlags::ClientContext | EAutomationTestFlags::EditorContext | EAutomationTestFlags::ProductFilter))
+	/**
+	 * Execute the loading of each map and performance captures
+	 *
+	 * @param Parameters - Should specify which map name to load
+	 * @return	TRUE if the test was successful, FALSE otherwise
+	 */
+	virtual bool RunTest(const FString& Parameters) override
+	{
+		FString MapObjectPath, MapPackageName, MapTestName;
+		ParseTestMapInfo(Parameters, MapObjectPath, MapPackageName, MapTestName);
 
+		bool bCanProceed = false;
 
-/** 
- * Requests a enumeration of all maps to be loaded
- */
-void FClientFunctionalTestingMaps::GetTests(TArray<FString>& OutBeautifiedNames, TArray<FString>& OutTestCommands) const
-{
-	IAssetRegistry& AssetRegistry = FModuleManager::Get().LoadModuleChecked<FAssetRegistryModule>("AssetRegistry").Get();
+		IFunctionalTestingModule::Get().MarkPendingActivation();
 
-	if ( !AssetRegistry.IsLoadingAssets() )
-	{
-		TArray<FAssetData> MapList;
-		FARFilter Filter;
-		Filter.ClassNames.Add(UWorld::StaticClass()->GetFName());
-		Filter.bRecursiveClasses = true;
-		Filter.bIncludeOnlyOnDiskAssets = true;
-		if ( AssetRegistry.GetAssets(Filter, /*out*/ MapList) )
+		UWorld* TestWorld = GetAnyGameWorld();
+		if (TestWorld && TestWorld->GetMapName() == MapPackageName)
 		{
-			for ( const FAssetData& MapAsset : MapList )
-			{
-				const FString* Tests = MapAsset.TagsAndValues.Find(TEXT("Tests"));
-				const FString* TestNames = MapAsset.TagsAndValues.Find(TEXT("TestNames"));
-
-				if ( Tests && TestNames )
-				{
-					int32 TestCount = FCString::Atoi(**Tests);
-					if ( TestCount > 0 )
-					{
-						TArray<FString> MapTests;
-						( *TestNames ).ParseIntoArray(MapTests, TEXT(";"), true);
-
-						for ( const FString& MapTest : MapTests )
-						{
-							FString BeautifulTestName;
-							FString RealTestName;
-
-							if ( MapTest.Split(TEXT("|"), &BeautifulTestName, &RealTestName) )
-							{
-								OutBeautifiedNames.Add(MapAsset.PackageName.ToString() + TEXT(".") + *BeautifulTestName);
-								OutTestCommands.Add(MapAsset.ObjectPath.ToString() + TEXT(";") + MapAsset.PackageName.ToString() + TEXT(";") + *RealTestName);
-							}
-						}
-					}
-				}
-				else if ( MapAsset.AssetName.ToString().Find(TEXT("FTEST_")) == 0 )
-				{
-					OutBeautifiedNames.Add(MapAsset.AssetName.ToString());
-					OutTestCommands.Add(MapAsset.ObjectPath.ToString() + TEXT(";") + MapAsset.PackageName.ToString());
-				}
-			}
+			// Map is already loaded.
+			bCanProceed = true;
 		}
-	}
-}
-
-
-// @todo this is a temporary solution. Once we know how to get test's hands on a proper world
-// this function should be redone/removed
-static UWorld* GetAnyGameWorld()
-{
-	UWorld* TestWorld = nullptr;
-	const TIndirectArray<FWorldContext>& WorldContexts = GEngine->GetWorldContexts();
-	for ( const FWorldContext& Context : WorldContexts )
-	{
-		if ( ( ( Context.WorldType == EWorldType::PIE ) || ( Context.WorldType == EWorldType::Game ) ) && ( Context.World() != NULL ) )
+		else
 		{
-			TestWorld = Context.World();
-			break;
+			bCanProceed = AutomationOpenMap(MapPackageName);
 		}
-	}
-
-	return TestWorld;
-}
 
+		if (bCanProceed)
+		{
+			if (MapTestName.IsEmpty())
+			{
+				ADD_LATENT_AUTOMATION_COMMAND(FStartFTestsOnMap());
+			}
+			else
+			{
+				ADD_LATENT_AUTOMATION_COMMAND(FStartFTestOnMap(MapTestName));
+			}
 
-/** 
- * Execute the loading of each map and performance captures
- *
- * @param Parameters - Should specify which map name to load
- * @return	TRUE if the test was successful, FALSE otherwise
- */
-bool FClientFunctionalTestingMaps::RunTest(const FString& Parameters)
-{
-	FString MapObjectPath, MapPackageName, MapTestName;
-	ParseTestMapInfo(Parameters, MapObjectPath, MapPackageName, MapTestName);
+			return true;
+		}
 
-	bool bCanProceed = false;
+		/// FAutomationTestFramework::GetInstance().UnregisterAutomationTest
 
-	FFunctionalTestingModule::Get()->MarkPendingActivation();
+		//	ADD_LATENT_AUTOMATION_COMMAND(FWaitLatentCommand(1.f));
+		//  ADD_LATENT_AUTOMATION_COMMAND(FExitGameCommand);
 
-	UWorld* TestWorld = GetAnyGameWorld();
-	if ( TestWorld && TestWorld->GetMapName() == MapPackageName )
-	{
-		// Map is already loaded.
-		bCanProceed = true;
-	}
-	else
-	{
-		bCanProceed = AutomationOpenMap(MapPackageName);
+		UE_LOG(LogFunctionalTest, Error, TEXT("Failed to start the %s map (possibly due to BP compilation issues)"), *MapPackageName);
+		return false;
 	}
+};
 
-	if (bCanProceed)
-	{
-		if ( MapTestName.IsEmpty() )
-		{
-			ADD_LATENT_AUTOMATION_COMMAND(FStartFTestsOnMap());
-		}
-		else
-		{
-			ADD_LATENT_AUTOMATION_COMMAND(FStartFTestOnMap(MapTestName));
-		}
+// Runtime tests
+IMPLEMENT_CUSTOM_COMPLEX_AUTOMATION_TEST(FClientFunctionalTestingMapsRuntime, FClientFunctionalTestingMapsBase, "Project.Functional Tests", (EAutomationTestFlags::ClientContext | EAutomationTestFlags::EditorContext | EAutomationTestFlags::ProductFilter))
 
-		return true;
-	}
+void FClientFunctionalTestingMapsRuntime::GetTests(TArray<FString>& OutBeautifiedNames, TArray<FString>& OutTestCommands) const
+{
+	FClientFunctionalTestingMapsBase::GetTests(OutBeautifiedNames, OutTestCommands);
+}
+
+bool FClientFunctionalTestingMapsRuntime::RunTest(const FString& Parameters)
+{
+	return FClientFunctionalTestingMapsBase::RunTest(Parameters);
+}
 
-	/// FAutomationTestFramework::GetInstance().UnregisterAutomationTest
+// Editor only tests
+IMPLEMENT_CUSTOM_COMPLEX_AUTOMATION_TEST(FClientFunctionalTestingMapsEditor, FClientFunctionalTestingMapsBase, "Project.Functional Tests", (EAutomationTestFlags::EditorContext | EAutomationTestFlags::ProductFilter))
 
-	//	ADD_LATENT_AUTOMATION_COMMAND(FWaitLatentCommand(1.f));
-	//  ADD_LATENT_AUTOMATION_COMMAND(FExitGameCommand);
+void FClientFunctionalTestingMapsEditor::GetTests(TArray<FString>& OutBeautifiedNames, TArray<FString>& OutTestCommands) const
+{
+	FClientFunctionalTestingMapsBase::GetTests(OutBeautifiedNames, OutTestCommands);
+}
 
-	UE_LOG(LogFunctionalTesting, Error, TEXT("Failed to start the %s map (possibly due to BP compilation issues)"), *MapPackageName);
-	return false;
+bool FClientFunctionalTestingMapsEditor::RunTest(const FString& Parameters)
+{
+	return FClientFunctionalTestingMapsBase::RunTest(Parameters);
 }
 
 #undef LOCTEXT_NAMESPACE
diff --git a/Engine/Source/Developer/FunctionalTesting/Public/FunctionalTestingModule.h b/Engine/Source/Developer/FunctionalTesting/Public/FunctionalTestingModule.h
--- a/Engine/Source/Developer/FunctionalTesting/Public/FunctionalTestingModule.h
+++ b/Engine/Source/Developer/FunctionalTesting/Public/FunctionalTestingModule.h
@@ -3,32 +3,48 @@
 #pragma once
 
 #include "CoreMinimal.h"
-#include "IFuncTestManager.h"
 #include "Modules/ModuleInterface.h"
 #include "Modules/ModuleManager.h"
 
-class FFunctionalTestingModule : public IModuleInterface
+DECLARE_LOG_CATEGORY_EXTERN(LogFunctionalTest, Log, All);
+
+class IFunctionalTestingModule : public IModuleInterface
 {
 public:
+	/** Triggers in sequence all functional tests found on the level */
+	virtual void RunAllTestsOnMap(bool bClearLog, bool bRunLooped) = 0;
 
-	virtual void StartupModule() override;
-	virtual void ShutdownModule() override;
+	/** Runs a single test */
+	virtual void RunTestOnMap(const FString& TestName, bool bClearLog, bool bRunLooped) = 0;
 
-	/** Gets the debugger singleton or returns NULL */
-	static class IFuncTestManager* Get()
-	{
-		FFunctionalTestingModule& Module = FModuleManager::Get().LoadModuleChecked<FFunctionalTestingModule>("FunctionalTesting");
-		return Module.GetSingleton();
-	}
+	/** Sets that a test is being started */
+	virtual void MarkPendingActivation() = 0;
 
-private:
-	void OnModulesChanged(FName Module, EModuleChangeReason Reason);
+	/** True if a test is about to start */
+	virtual bool IsActivationPending() const = 0;
 
-	class IFuncTestManager* GetSingleton() const 
-	{ 
-		return Manager.Get(); 
-	}
+	/** True if there is an active UFunctionalTestingManager */
+	virtual bool IsRunning() const = 0;
+
+	/** True if a UFunctionalTestingManager was spawned and is now done */
+	virtual bool IsFinished() const = 0;
+
+	/** Sets the active testing manager */
+	virtual void SetManager(class UFunctionalTestingManager* NewManager) = 0;
 
-private:
-	TSharedPtr<class IFuncTestManager> Manager;
+	/** Getst he active testing manager */
+	virtual class UFunctionalTestingManager* GetCurrentManager() = 0;
+
+	/** If true, will run tests forever */
+	virtual void SetLooping(const bool bLoop) = 0;
+
+	/** Gets a list of maps/tests in the current project */
+	virtual void GetMapTests(bool bEditorOnlyTests, TArray<FString>& OutBeautifiedNames, TArray<FString>& OutTestCommands, TArray<FString>& OutTestMapAssets) const = 0;
+
+	/** Gets the debugger singleton or returns NULL */
+	static IFunctionalTestingModule& Get()
+	{
+		static const FName FunctionalTesting(TEXT("FunctionalTesting"));
+		return FModuleManager::Get().LoadModuleChecked<IFunctionalTestingModule>(FunctionalTesting);
+	}
 };
diff --git a/Engine/Source/Developer/HotReload/Private/HotReload.cpp b/Engine/Source/Developer/HotReload/Private/HotReload.cpp
--- a/Engine/Source/Developer/HotReload/Private/HotReload.cpp
+++ b/Engine/Source/Developer/HotReload/Private/HotReload.cpp
@@ -791,41 +791,10 @@ ECompilationResult::Type FHotReloadModule::DoHotReloadFromEditor(const bool bWai
 	return Result;
 }
 
-#if WITH_HOT_RELOAD
-/**
- * Gets duplicated CDO from the cache, renames it and returns.
- */
-UObject* GetCachedCDODuplicate(UObject* CDO, FName Name)
-{
-	UObject* DupCDO = nullptr;
-
-	UObject** DupCDOPtr = GetDuplicatedCDOMap().Find(CDO);
-	if (DupCDOPtr != nullptr)
-	{
-		DupCDO = *DupCDOPtr;
-		DupCDO->Rename(*Name.ToString(), GetTransientPackage(),
-			REN_DoNotDirty | REN_DontCreateRedirectors | REN_ForceNoResetLoaders | REN_NonTransactional | REN_SkipGeneratedClasses);
-	}
-
-	return DupCDO;
-}
-#endif // WITH_HOT_RELOAD
-
 ECompilationResult::Type FHotReloadModule::DoHotReloadInternal(const TMap<FString, FString>& ChangedModules, const TArray<UPackage*>& Packages, const TArray<FName>& InDependentModules, FOutputDevice& HotReloadAr)
 {
 #if WITH_HOT_RELOAD
 
-#if WITH_ENGINE
-	// Register with the BlueprintCompileReinstancer to handle duplicate requests
-	FBlueprintCompileReinstancer::FCDODuplicatesProvider& CDODuplicatesProvider = FBlueprintCompileReinstancer::GetCDODuplicatesProviderDelegate();
-	CDODuplicatesProvider.BindStatic(&GetCachedCDODuplicate);
-	ON_SCOPE_EXIT
-	{
-		CDODuplicatesProvider.Unbind();
-		GetDuplicatedCDOMap().Empty();
-	};
-#endif
-
 	FModuleManager& ModuleManager = FModuleManager::Get();
 
 	ModuleManager.ResetModulePathsCache();
diff --git a/Engine/Source/Developer/HotReload/Private/HotReloadClassReinstancer.cpp b/Engine/Source/Developer/HotReload/Private/HotReloadClassReinstancer.cpp
--- a/Engine/Source/Developer/HotReload/Private/HotReloadClassReinstancer.cpp
+++ b/Engine/Source/Developer/HotReload/Private/HotReloadClassReinstancer.cpp
@@ -161,14 +161,14 @@ void FHotReloadClassReinstancer::SerializeCDOProperties(UObject* InObject, FHotR
 			Ar << UniqueID;
 			return *this;
 		}
-		virtual FArchive& operator<<(FAssetPtr& AssetPtr) override
+		virtual FArchive& operator<<(FSoftObjectPtr& Value) override
 		{
 			FArchive& Ar = *this;
-			FStringAssetReference UniqueID = AssetPtr.GetUniqueID();
+			FSoftObjectPath UniqueID = Value.GetUniqueID();
 			Ar << UniqueID;
 			return Ar;
 		}
-		virtual FArchive& operator<<(FStringAssetReference& Value) override
+		virtual FArchive& operator<<(FSoftObjectPath& Value) override
 		{
 			FArchive& Ar = *this;
 
@@ -400,14 +400,14 @@ void FHotReloadClassReinstancer::UpdateDefaultProperties()
 			Ar << UniqueID;
 			return *this;
 		}
-		virtual FArchive& operator<<(FAssetPtr& AssetPtr) override
+		virtual FArchive& operator<<(FSoftObjectPtr& Value) override
 		{
 			FArchive& Ar = *this;
-			FStringAssetReference UniqueID = AssetPtr.GetUniqueID();
+			FSoftObjectPath UniqueID = Value.GetUniqueID();
 			Ar << UniqueID;
 			return Ar;
 		}
-		virtual FArchive& operator<<(FStringAssetReference& Value) override
+		virtual FArchive& operator<<(FSoftObjectPath& Value) override
 		{
 			FArchive& Ar = *this;
 
diff --git a/Engine/Source/Developer/MaterialUtilities/Private/MaterialUtilities.cpp b/Engine/Source/Developer/MaterialUtilities/Private/MaterialUtilities.cpp
--- a/Engine/Source/Developer/MaterialUtilities/Private/MaterialUtilities.cpp
+++ b/Engine/Source/Developer/MaterialUtilities/Private/MaterialUtilities.cpp
@@ -2,7 +2,6 @@
 #include "MaterialUtilities.h"
 #include "EngineDefines.h"
 #include "ShowFlags.h"
-#include "Misc/StringAssetReference.h"
 #include "Materials/MaterialInterface.h"
 #include "Materials/Material.h"
 #include "Engine/Texture2D.h"
@@ -2220,7 +2219,7 @@ bool FMaterialUtilities::ExportMaterialUVDensities(UMaterialInterface* InMateria
 				FMaterialTextureInfo TextureInfo;
 				TextureInfo.SamplingScale = SamplingScale;
 				TextureInfo.UVChannelIndex = CoordIndex;
-				TextureInfo.TextureReference = FStringAssetReference(Texture2D);
+				TextureInfo.TextureReference = FSoftObjectPath(Texture2D);
 				TextureInfo.TextureIndex = RegisterIndex;
 				TextureStreamingData.Add(TextureInfo);
 				bSuccess = true;
diff --git a/Engine/Source/Editor/AdvancedPreviewScene/Public/AssetViewerSettings.h b/Engine/Source/Editor/AdvancedPreviewScene/Public/AssetViewerSettings.h
--- a/Engine/Source/Editor/AdvancedPreviewScene/Public/AssetViewerSettings.h
+++ b/Engine/Source/Editor/AdvancedPreviewScene/Public/AssetViewerSettings.h
@@ -5,7 +5,7 @@
 #include "CoreMinimal.h"
 #include "UObject/ObjectMacros.h"
 #include "UObject/Object.h"
-#include "UObject/AssetPtr.h"
+#include "UObject/SoftObjectPtr.h"
 #include "UObject/ObjectRedirector.h"
 #include "Engine/Scene.h"
 #include "Engine/TextureCube.h"
@@ -73,7 +73,7 @@ struct FPreviewSceneProfile
 
 	/** Sets environment cube map used for sky lighting and reflections */
 	UPROPERTY(EditAnywhere, transient, Category = Environment)
-	TAssetPtr<UTextureCube> EnvironmentCubeMap;
+	TSoftObjectPtr<UTextureCube> EnvironmentCubeMap;
 
 	/** Storing path to environment cube to prevent it from getting cooked */
 	UPROPERTY(config)
diff --git a/Engine/Source/Editor/AnimGraph/Classes/AnimGraphNode_AssetPlayerBase.h b/Engine/Source/Editor/AnimGraph/Classes/AnimGraphNode_AssetPlayerBase.h
--- a/Engine/Source/Editor/AnimGraph/Classes/AnimGraphNode_AssetPlayerBase.h
+++ b/Engine/Source/Editor/AnimGraph/Classes/AnimGraphNode_AssetPlayerBase.h
@@ -4,7 +4,6 @@
 
 #include "CoreMinimal.h"
 #include "UObject/ObjectMacros.h"
-#include "Misc/StringAssetReference.h"
 #include "Animation/AnimationAsset.h"
 #include "AnimGraphNode_Base.h"
 #include "AnimGraphNode_AssetPlayerBase.generated.h"
diff --git a/Engine/Source/Editor/AnimGraph/Classes/AnimGraphNode_Base.h b/Engine/Source/Editor/AnimGraph/Classes/AnimGraphNode_Base.h
--- a/Engine/Source/Editor/AnimGraph/Classes/AnimGraphNode_Base.h
+++ b/Engine/Source/Editor/AnimGraph/Classes/AnimGraphNode_Base.h
@@ -156,6 +156,8 @@ class ANIMGRAPH_API UAnimGraphNode_Base : public UK2Node
 
 	// By default return any animation assets we have
 	virtual UObject* GetJumpTargetForDoubleClick() const override { return GetAnimationAsset(); }
+	virtual bool CanJumpToDefinition() const override;
+	virtual void JumpToDefinition() const override;
 	// End of UK2Node interface
 
 	// UAnimGraphNode_Base interface
diff --git a/Engine/Source/Editor/AnimGraph/Classes/AnimStateNodeBase.h b/Engine/Source/Editor/AnimGraph/Classes/AnimStateNodeBase.h
--- a/Engine/Source/Editor/AnimGraph/Classes/AnimStateNodeBase.h
+++ b/Engine/Source/Editor/AnimGraph/Classes/AnimStateNodeBase.h
@@ -24,6 +24,8 @@ class UAnimStateNodeBase : public UEdGraphNode
 	// UEdGraphNode interface
 	virtual void PostPasteNode() override;
 	virtual UObject* GetJumpTargetForDoubleClick() const override;
+	virtual bool CanJumpToDefinition() const override;
+	virtual void JumpToDefinition() const override;
 	virtual bool CanCreateUnderSpecifiedSchema(const UEdGraphSchema* Schema) const override;
 	virtual void OnRenameNode(const FString& NewName) override;
 	virtual TSharedPtr<class INameValidatorInterface> MakeNameValidator() const override;
diff --git a/Engine/Source/Editor/AnimGraph/Private/AnimGraphNode_Base.cpp b/Engine/Source/Editor/AnimGraph/Private/AnimGraphNode_Base.cpp
--- a/Engine/Source/Editor/AnimGraph/Private/AnimGraphNode_Base.cpp
+++ b/Engine/Source/Editor/AnimGraph/Private/AnimGraphNode_Base.cpp
@@ -12,6 +12,7 @@
 #include "AnimNodeEditModes.h"
 #include "AnimationGraph.h"
 #include "EditorModeManager.h"
+#include "Toolkits/AssetEditorManager.h"
 
 /////////////////////////////////////////////////////
 // UAnimGraphNode_Base
@@ -87,6 +88,19 @@ void UAnimGraphNode_Base::ReallocatePinsDuringReconstruction(TArray<UEdGraphPin*
 	RestoreSplitPins(OldPins);
 }
 
+bool UAnimGraphNode_Base::CanJumpToDefinition() const
+{
+	return GetJumpTargetForDoubleClick() != nullptr;
+}
+
+void UAnimGraphNode_Base::JumpToDefinition() const
+{
+	if (UObject* HyperlinkTarget = GetJumpTargetForDoubleClick())
+	{
+		FAssetEditorManager::Get().OpenEditorForAsset(HyperlinkTarget);
+	}
+}
+
 FLinearColor UAnimGraphNode_Base::GetNodeTitleColor() const
 {
 	return FLinearColor::Black;
diff --git a/Engine/Source/Editor/AnimGraph/Private/AnimStateNodeBase.cpp b/Engine/Source/Editor/AnimGraph/Private/AnimStateNodeBase.cpp
--- a/Engine/Source/Editor/AnimGraph/Private/AnimStateNodeBase.cpp
+++ b/Engine/Source/Editor/AnimGraph/Private/AnimStateNodeBase.cpp
@@ -11,6 +11,7 @@
 #include "AnimationStateMachineGraph.h"
 #include "AnimationStateMachineSchema.h"
 #include "Kismet2/Kismet2NameValidators.h"
+#include "Kismet2/KismetEditorUtilities.h"
 /////////////////////////////////////////////////////
 // FAnimStateNodeNameValidator
 
@@ -74,6 +75,19 @@ UObject* UAnimStateNodeBase::GetJumpTargetForDoubleClick() const
 	return GetBoundGraph();
 }
 
+bool UAnimStateNodeBase::CanJumpToDefinition() const
+{
+	return GetJumpTargetForDoubleClick() != nullptr;
+}
+
+void UAnimStateNodeBase::JumpToDefinition() const
+{
+	if (UObject* HyperlinkTarget = GetJumpTargetForDoubleClick())
+	{
+		FKismetEditorUtilities::BringKismetToFocusAttentionOnObject(HyperlinkTarget);
+	}
+}
+
 bool UAnimStateNodeBase::CanCreateUnderSpecifiedSchema(const UEdGraphSchema* Schema) const
 {
 	return Schema->IsA(UAnimationStateMachineSchema::StaticClass());
diff --git a/Engine/Source/Editor/AnimationBlueprintEditor/Private/AnimationBlueprintEditorModule.cpp b/Engine/Source/Editor/AnimationBlueprintEditor/Private/AnimationBlueprintEditorModule.cpp
--- a/Engine/Source/Editor/AnimationBlueprintEditor/Private/AnimationBlueprintEditorModule.cpp
+++ b/Engine/Source/Editor/AnimationBlueprintEditor/Private/AnimationBlueprintEditorModule.cpp
@@ -74,7 +74,7 @@ void FAnimationBlueprintEditorModule::OnNewBlueprintCreated(UBlueprint* InBluepr
 		GetOwnerNode->NodePosX = SafeXPosition;
 		GetOwnerNode->NodePosY = SafeYPosition;
 		UEdGraphSchema_K2::SetNodeMetaData(GetOwnerNode, FNodeMetadata::DefaultGraphNode);
-		GetOwnerNode->DisableNode();
+		GetOwnerNode->MakeAutomaticallyPlacedGhostNode();
 
 		EventGraph->AddNode(GetOwnerNode);
 	}
diff --git a/Engine/Source/Editor/BlueprintGraph/Classes/EdGraphSchema_K2.h b/Engine/Source/Editor/BlueprintGraph/Classes/EdGraphSchema_K2.h
--- a/Engine/Source/Editor/BlueprintGraph/Classes/EdGraphSchema_K2.h
+++ b/Engine/Source/Editor/BlueprintGraph/Classes/EdGraphSchema_K2.h
@@ -6,7 +6,7 @@
 #include "Misc/EnumClassFlags.h"
 #include "UObject/ObjectMacros.h"
 #include "UObject/Class.h"
-#include "Misc/StringAssetReference.h"
+#include "UObject/SoftObjectPath.h"
 #include "EdGraph/EdGraphNode.h"
 #include "EdGraph/EdGraphPin.h"
 #include "AssetData.h"
@@ -249,8 +249,8 @@ enum class EObjectReferenceType : uint8
 	NotAnObject		= 0x00,
 	ObjectReference = 0x01,
 	ClassReference	= 0x02,
-	AssetID			= 0x04,
-	ClassAssetID	= 0x08,
+	SoftObject		= 0x04,
+	SoftClass		= 0x08,
 	AllTypes		= 0x0f,
 };
 
@@ -280,7 +280,7 @@ class BLUEPRINTGRAPH_API UEdGraphSchema_K2 : public UEdGraphSchema
 	static const FString PC_Boolean;
 	static const FString PC_Byte;
 	static const FString PC_Class;    // SubCategoryObject is the MetaClass of the Class passed thru this pin, or SubCategory can be 'self'. The DefaultValue string should always be empty, use DefaultObject.
-	static const FString PC_AssetClass;
+	static const FString PC_SoftClass;
 	static const FString PC_Int;
 	static const FString PC_Float;
 	static const FString PC_Name;
@@ -288,7 +288,7 @@ class BLUEPRINTGRAPH_API UEdGraphSchema_K2 : public UEdGraphSchema
 	static const FString PC_MCDelegate;  // SubCategoryObject is the UFunction of the delegate signature
 	static const FString PC_Object;    // SubCategoryObject is the Class of the object passed thru this pin, or SubCategory can be 'self'. The DefaultValue string should always be empty, use DefaultObject.
 	static const FString PC_Interface;	// SubCategoryObject is the Class of the object passed thru this pin.
-	static const FString PC_Asset;		// SubCategoryObject is the Class of the AssetPtr passed thru this pin.
+	static const FString PC_SoftObject;		// SubCategoryObject is the Class of the SoftObjectPtr passed thru this pin.
 	static const FString PC_String;
 	static const FString PC_Text;
 	static const FString PC_Struct;    // SubCategoryObject is the ScriptStruct of the struct passed thru this pin, 'self' is not a valid SubCategory. DefaultObject should always be empty, the DefaultValue string may be used for supported structs.
@@ -367,7 +367,7 @@ public:
 		uint8 PossibleObjectReferenceTypes;
 
 		/** Asset Reference, used when PinType.PinSubCategoryObject is not loaded yet */
-		FStringAssetReference SubCategoryObjectAssetReference;
+		FSoftObjectPath SubCategoryObjectAssetReference;
 
 		FText CachedDescription;
 
@@ -393,7 +393,7 @@ public:
 
 		FPinTypeTreeInfo(const FText& InFriendlyName, const FString& CategoryName, const UEdGraphSchema_K2* Schema, const FText& InTooltip, bool bInReadOnly = false, FTypesDatabase* TypesDatabase = nullptr);
 		FPinTypeTreeInfo(const FString& CategoryName, UObject* SubCategoryObject, const FText& InTooltip, bool bInReadOnly = false, uint8 InPossibleObjectReferenceTypes = 0);
-		FPinTypeTreeInfo(const FText& InFriendlyName, const FString& CategoryName, const FStringAssetReference& SubCategoryObject, const FText& InTooltip, bool bInReadOnly = false, uint8 InPossibleObjectReferenceTypes = 0);
+		FPinTypeTreeInfo(const FText& InFriendlyName, const FString& CategoryName, const FSoftObjectPath& SubCategoryObject, const FText& InTooltip, bool bInReadOnly = false, uint8 InPossibleObjectReferenceTypes = 0);
 
 		FPinTypeTreeInfo(TSharedPtr<FPinTypeTreeInfo> InInfo)
 		{
@@ -481,13 +481,13 @@ public:
 	virtual void GetAssetsGraphHoverMessage(const TArray<FAssetData>& Assets, const UEdGraph* HoverGraph, FString& OutTooltipText, bool& OutOkIcon) const override;
 	virtual bool CanDuplicateGraph(UEdGraph* InSourceGraph) const override;
 	virtual UEdGraph* DuplicateGraph(UEdGraph* GraphToDuplicate) const override;
-	virtual UEdGraphNode* CreateSubstituteNode(UEdGraphNode* Node, const UEdGraph* Graph, FObjectInstancingGraph* InstanceGraph, TArray<FName>& InOutExtraNames) const override;
+	virtual UEdGraphNode* CreateSubstituteNode(UEdGraphNode* Node, const UEdGraph* Graph, FObjectInstancingGraph* InstanceGraph, TSet<FName>& InOutExtraNames) const override;
 	virtual int32 GetNodeSelectionCount(const UEdGraph* Graph) const override;
 	virtual TSharedPtr<FEdGraphSchemaAction> GetCreateCommentAction() const override;
 	virtual bool FadeNodeWhenDraggingOffPin(const UEdGraphNode* Node, const UEdGraphPin* Pin) const override;
 	virtual void BackwardCompatibilityNodeConversion(UEdGraph* Graph, bool bOnlySafeChanges) const override;
 	virtual bool ShouldAlwaysPurgeOnModification() const override { return false; }
-	virtual void SplitPin(UEdGraphPin* Pin) const override;
+	virtual void SplitPin(UEdGraphPin* Pin, bool bNotify = true) const override;
 	virtual void RecombinePin(UEdGraphPin* Pin) const override;
 	virtual void OnPinConnectionDoubleCicked(UEdGraphPin* PinA, UEdGraphPin* PinB, const FVector2D& GraphPosition) const override;
 	virtual UEdGraphPin* DropPinOnNode(UEdGraphNode* InTargetNode, const FString& InSourcePinName, const FEdGraphPinType& InSourcePinType, EEdGraphPinDirection InSourcePinDirection) const override;
diff --git a/Engine/Source/Editor/BlueprintGraph/Classes/EdGraphSchema_K2_Actions.h b/Engine/Source/Editor/BlueprintGraph/Classes/EdGraphSchema_K2_Actions.h
--- a/Engine/Source/Editor/BlueprintGraph/Classes/EdGraphSchema_K2_Actions.h
+++ b/Engine/Source/Editor/BlueprintGraph/Classes/EdGraphSchema_K2_Actions.h
@@ -200,33 +200,6 @@ struct FEdGraphSchemaAction_K2AddComponent : public FEdGraphSchemaAction_K2NewNo
 	// End of FEdGraphSchemaAction interface
 };
 
-/*******************************************************************************
-* FEdGraphSchemaAction_K2AddTimeline
-*******************************************************************************/
-
-/** Action to add a 'timeline' node to the graph */
-USTRUCT()
-struct FEdGraphSchemaAction_K2AddTimeline : public FEdGraphSchemaAction_K2NewNode
-{
-	GENERATED_USTRUCT_BODY()
-
-	// Simple type info
-	static FName StaticGetTypeId() {static FName Type("FEdGraphSchemaAction_K2AddTimeline"); return Type;}
-	virtual FName GetTypeId() const override { return StaticGetTypeId(); } 
-
-	FEdGraphSchemaAction_K2AddTimeline()
-		: FEdGraphSchemaAction_K2NewNode()
-	{}
-
-	FEdGraphSchemaAction_K2AddTimeline(FText InNodeCategory, FText InMenuDesc, FText InToolTip, const int32 InGrouping)
-		: FEdGraphSchemaAction_K2NewNode(MoveTemp(InNodeCategory), MoveTemp(InMenuDesc), MoveTemp(InToolTip), InGrouping)
-	{}
-
-	// FEdGraphSchemaAction interface
-	virtual UEdGraphNode* PerformAction(class UEdGraph* ParentGraph, UEdGraphPin* FromPin, const FVector2D Location, bool bSelectNewNode = true) override;
-	// End of FEdGraphSchemaAction interface
-};
-
 /*******************************************************************************
 * FEdGraphSchemaAction_K2AddEvent
 *******************************************************************************/
@@ -441,44 +414,44 @@ struct BLUEPRINTGRAPH_API FEdGraphSchemaAction_K2Enum : public FEdGraphSchemaAct
 };
 
 /*******************************************************************************
-* FEdGraphSchemaAction_K2Var
+* FEdGraphSchemaAction_BlueprintVariableBase
 *******************************************************************************/
 
-/** Reference to a variable (only used in 'docked' palette) */
+/** Reference to a variable (only used in 'My Blueprints' but used for member variables, local variables, delegates, etc...) */
 USTRUCT()
-struct BLUEPRINTGRAPH_API FEdGraphSchemaAction_K2Var : public FEdGraphSchemaAction
+struct BLUEPRINTGRAPH_API FEdGraphSchemaAction_BlueprintVariableBase : public FEdGraphSchemaAction
 {
 	GENERATED_USTRUCT_BODY()
 
 private:
 	/** Name of function or class */
-	FName		VarName;
+	FName VarName;
 
-	/** The class that owns this item */
-	TWeakObjectPtr<UClass>		OwningClass;
+	/** The struct that owns this item */
+	TWeakObjectPtr<UStruct> VariableSource;
 
 	/** TRUE if the variable's type is boolean */
 	bool bIsVarBool;
 
 public:
-	void SetVariableInfo(const FName& InVarName, const UClass* InOwningClass, bool bInIsVarBool)
+	void SetVariableInfo(const FName& InVarName, const UStruct* InOwningScope, bool bInIsVarBool)
 	{
 		VarName = InVarName;
 		bIsVarBool = bInIsVarBool;
 
-		check(InOwningClass);
-		OwningClass = InOwningClass;
+		check(InOwningScope);
+		VariableSource = InOwningScope;
 	}
 
 	// Simple type info
 	static FName StaticGetTypeId() {static FName Type("FEdGraphSchemaAction_K2Var"); return Type;}
 	virtual FName GetTypeId() const override { return StaticGetTypeId(); } 
 
-	FEdGraphSchemaAction_K2Var() 
+	FEdGraphSchemaAction_BlueprintVariableBase()
 		: FEdGraphSchemaAction()
 	{}
 
-	FEdGraphSchemaAction_K2Var(FText InNodeCategory, FText InMenuDesc, FText InToolTip, const int32 InGrouping, const int32 InSectionID)
+	FEdGraphSchemaAction_BlueprintVariableBase(FText InNodeCategory, FText InMenuDesc, FText InToolTip, const int32 InGrouping, const int32 InSectionID)
 		: FEdGraphSchemaAction(MoveTemp(InNodeCategory), MoveTemp(InMenuDesc), MoveTemp(InToolTip), InGrouping, FText(), InSectionID)
 	{}
 
@@ -494,84 +467,76 @@ public:
 
 	UClass* GetVariableClass() const
 	{
-		return OwningClass.Get();
+		return Cast<UClass>(GetVariableScope());
 	}
 
-	UProperty* GetProperty() const
+	UStruct* GetVariableScope() const
 	{
-		UClass* SearchScope = GetVariableClass();
-		return FindField<UProperty>(SearchScope, VarName);
+		return VariableSource.Get();
 	}
 
+	UProperty* GetProperty() const
+	{
+		return FindField<UProperty>(GetVariableScope(), VarName);
+	}
+	
 	// FEdGraphSchemaAction interface
+	virtual void MovePersistentItemToCategory(const FText& NewCategoryName) override;
+	virtual int32 GetReorderIndexInContainer() const override;
+	virtual bool ReorderToBeforeAction(TSharedRef<FEdGraphSchemaAction> OtherAction) override;
+	virtual FEdGraphSchemaActionDefiningObject GetPersistentItemDefiningObject() const override;
 	// End of FEdGraphSchemaAction interface
+
+	UBlueprint* GetSourceBlueprint() const;
 };
 
+
 /*******************************************************************************
-* FEdGraphSchemaAction_K2LocalVar
+* FEdGraphSchemaAction_K2Var
 *******************************************************************************/
 
-/** Reference to a local variable (only used in 'docked' palette) */
+/** Reference to a variable (only used in 'docked' palette) */
 USTRUCT()
-struct BLUEPRINTGRAPH_API FEdGraphSchemaAction_K2LocalVar : public FEdGraphSchemaAction
+struct BLUEPRINTGRAPH_API FEdGraphSchemaAction_K2Var : public FEdGraphSchemaAction_BlueprintVariableBase
 {
 	GENERATED_USTRUCT_BODY()
 
-private:
-	/** Name of function or class */
-	FName		VarName;
+public:
+	// Simple type info
+	static FName StaticGetTypeId() { static FName Type("FEdGraphSchemaAction_K2Var"); return Type; }
+	virtual FName GetTypeId() const override { return StaticGetTypeId(); }
 
-	/** The struct that owns this item */
-	TWeakObjectPtr<UStruct>		VariableScope;
+	FEdGraphSchemaAction_K2Var()
+		: FEdGraphSchemaAction_BlueprintVariableBase()
+	{}
 
-	/** TRUE if the variable's type is boolean */
-	bool bIsVarBool;
+	FEdGraphSchemaAction_K2Var(FText InNodeCategory, FText InMenuDesc, FText InToolTip, const int32 InGrouping, const int32 InSectionID)
+		: FEdGraphSchemaAction_BlueprintVariableBase(MoveTemp(InNodeCategory), MoveTemp(InMenuDesc), MoveTemp(InToolTip), InGrouping, InSectionID)
+	{}
+};
 
-public:
-	void SetVariableInfo(const FName& InVarName, const UStruct* InVariableScope, bool bInIsVarBool)
-	{
-		VarName = InVarName;
-		bIsVarBool = bInIsVarBool;
+/*******************************************************************************
+* FEdGraphSchemaAction_K2LocalVar
+*******************************************************************************/
 
-		check(InVariableScope);
-		VariableScope = InVariableScope;
-	}
+/** Reference to a local variable (only used in 'docked' palette) */
+USTRUCT()
+struct BLUEPRINTGRAPH_API FEdGraphSchemaAction_K2LocalVar : public FEdGraphSchemaAction_BlueprintVariableBase
+{
+	GENERATED_USTRUCT_BODY()
 
+public:
 	// Simple type info
 	static FName StaticGetTypeId() {static FName Type("FEdGraphSchemaAction_K2LocalVar"); return Type;}
 	virtual FName GetTypeId() const override { return StaticGetTypeId(); } 
 
 	FEdGraphSchemaAction_K2LocalVar() 
-		: FEdGraphSchemaAction()
+		: FEdGraphSchemaAction_BlueprintVariableBase()
 	{}
 
 	FEdGraphSchemaAction_K2LocalVar(FText InNodeCategory, FText InMenuDesc, FText InToolTip, const int32 InGrouping, const int32 InSectionID)
-		: FEdGraphSchemaAction(MoveTemp(InNodeCategory), MoveTemp(InMenuDesc), MoveTemp(InToolTip), InGrouping, FText(), InSectionID)
+		: FEdGraphSchemaAction_BlueprintVariableBase(MoveTemp(InNodeCategory), MoveTemp(InMenuDesc), MoveTemp(InToolTip), InGrouping, InSectionID)
 	{}
-
-	FName GetVariableName() const
-	{
-		return VarName;
-	}
-
-	FString GetFriendlyVariableName() const
-	{
-		return FName::NameToDisplayString( VarName.ToString(), bIsVarBool );
-	}
-
-	UStruct* GetVariableScope() const
-	{
-		return VariableScope.Get();
-	}
-
-	UProperty* GetProperty() const
-	{
-		UStruct* SearchScope = GetVariableScope();
-		return FindField<UProperty>(SearchScope, VarName);
-	}
-
-	// FEdGraphSchemaAction interface
-	// End of FEdGraphSchemaAction interface
 };
 
 /*******************************************************************************
@@ -624,7 +589,15 @@ struct BLUEPRINTGRAPH_API FEdGraphSchemaAction_K2Graph : public FEdGraphSchemaAc
 
 	// FEdGraphSchemaAction interface
 	virtual bool IsParentable() const override { return true; }
+	virtual void MovePersistentItemToCategory(const FText& NewCategoryName) override;
+	virtual int32 GetReorderIndexInContainer() const override;
+	virtual bool ReorderToBeforeAction(TSharedRef<FEdGraphSchemaAction> OtherAction) override;
+	virtual FEdGraphSchemaActionDefiningObject GetPersistentItemDefiningObject() const override;
 	// End of FEdGraphSchemaAction interface
+
+protected:
+	UFunction* GetFunction() const;
+	UBlueprint* GetSourceBlueprint() const;
 };
 
 /*******************************************************************************
@@ -731,28 +704,13 @@ struct BLUEPRINTGRAPH_API FEdGraphSchemaAction_K2InputAction : public FEdGraphSc
 * FEdGraphSchemaAction_K2Delegate
 *******************************************************************************/
 
-/** Reference to a function, macro, event graph, or timeline (only used in 'docked' palette) */
+/** Reference to a delegate */
 USTRUCT()
-struct BLUEPRINTGRAPH_API FEdGraphSchemaAction_K2Delegate : public FEdGraphSchemaAction
+struct BLUEPRINTGRAPH_API FEdGraphSchemaAction_K2Delegate : public FEdGraphSchemaAction_BlueprintVariableBase
 {
 	GENERATED_USTRUCT_BODY()
 
-private:
-	/** Name of the property */
-	FName DelegateName;
-
-	/** Class that owns the delegates */
-	TWeakObjectPtr<UClass> OwningClass;
-
 public:		
-	void SetDelegateInfo(const FName& InDelegateName, const UClass* InDelegateClass)
-	{
-		DelegateName = InDelegateName;
-
-		check(InDelegateClass);
-		OwningClass = const_cast<UClass*>(InDelegateClass);
-	}
-
 	// Simple type info
 	static FName StaticGetTypeId() {static FName Type("FEdGraphSchemaAction_K2Delegate"); return Type;}
 	virtual FName GetTypeId() const override { return StaticGetTypeId(); } 
@@ -761,32 +719,29 @@ public:
 	UEdGraph* EdGraph;
 
 	FEdGraphSchemaAction_K2Delegate() 
-		: FEdGraphSchemaAction(), EdGraph(NULL)
+		: FEdGraphSchemaAction_BlueprintVariableBase()
+		, EdGraph(nullptr)
 	{}
 
 	FEdGraphSchemaAction_K2Delegate(FText InNodeCategory, FText InMenuDesc, FText InToolTip, const int32 InGrouping, const int32 InSectionID)
-		: FEdGraphSchemaAction(MoveTemp(InNodeCategory), MoveTemp(InMenuDesc), MoveTemp(InToolTip), InGrouping, FText(), InSectionID)
+		: FEdGraphSchemaAction_BlueprintVariableBase(MoveTemp(InNodeCategory), MoveTemp(InMenuDesc), MoveTemp(InToolTip), InGrouping, InSectionID)
 		, EdGraph(nullptr)
 	{}
 
 	FName GetDelegateName() const
 	{
-		return DelegateName;
+		return GetVariableName();
 	}
 
 	UClass* GetDelegateClass() const
 	{
-		return OwningClass.Get();
+		return GetVariableClass();
 	}
 
 	UMulticastDelegateProperty* GetDelegateProperty() const
 	{
-		UClass* DelegateClass = GetDelegateClass();
-		return FindField<UMulticastDelegateProperty>(DelegateClass, DelegateName);
+		return FindField<UMulticastDelegateProperty>(GetVariableClass(), GetVariableName());
 	}
-
-	// FEdGraphSchemaAction interface
-	// End of FEdGraphSchemaAction interface
 };
 
 
diff --git a/Engine/Source/Editor/BlueprintGraph/Classes/K2Node.h b/Engine/Source/Editor/BlueprintGraph/Classes/K2Node.h
--- a/Engine/Source/Editor/BlueprintGraph/Classes/K2Node.h
+++ b/Engine/Source/Editor/BlueprintGraph/Classes/K2Node.h
@@ -182,6 +182,8 @@ class UK2Node : public UEdGraphNode
 	BLUEPRINTGRAPH_API virtual void AutowireNewNode(UEdGraphPin* FromPin) override;
 	BLUEPRINTGRAPH_API void PinConnectionListChanged(UEdGraphPin* Pin) override;
 	BLUEPRINTGRAPH_API virtual UObject* GetJumpTargetForDoubleClick() const override;
+	BLUEPRINTGRAPH_API virtual bool CanJumpToDefinition() const override;
+	BLUEPRINTGRAPH_API virtual void JumpToDefinition() const override;
 	BLUEPRINTGRAPH_API virtual FString GetDocumentationLink() const override;
 	BLUEPRINTGRAPH_API virtual void GetPinHoverText(const UEdGraphPin& Pin, FString& HoverTextOut) const override;
 	BLUEPRINTGRAPH_API virtual bool ShowPaletteIconOnNode() const override { return true; }
diff --git a/Engine/Source/Editor/BlueprintGraph/Classes/K2Node_ActorBoundEvent.h b/Engine/Source/Editor/BlueprintGraph/Classes/K2Node_ActorBoundEvent.h
--- a/Engine/Source/Editor/BlueprintGraph/Classes/K2Node_ActorBoundEvent.h
+++ b/Engine/Source/Editor/BlueprintGraph/Classes/K2Node_ActorBoundEvent.h
@@ -48,6 +48,7 @@ class UK2Node_ActorBoundEvent : public UK2Node_Event
 	//~ Begin K2Node Interface
 	virtual AActor* GetReferencedLevelActor() const override;
 	virtual bool NodeCausesStructuralBlueprintChange() const override { return true; }
+	virtual void ValidateNodeDuringCompilation(class FCompilerResultsLog& MessageLog) const override;
 	virtual FNodeHandlingFunctor* CreateNodeHandler(FKismetCompilerContext& CompilerContext) const override;
 	//~ End K2Node Interface
 
@@ -63,6 +64,7 @@ class UK2Node_ActorBoundEvent : public UK2Node_Event
 
 	/** Return the delegate property that this event is bound to */
 	BLUEPRINTGRAPH_API UMulticastDelegateProperty* GetTargetDelegateProperty() const;
+	BLUEPRINTGRAPH_API UMulticastDelegateProperty* GetTargetDelegatePropertyFromSkel() const;
 
 	/** Returns the delegate that this event is bound in, for actors bound in level script */
 	BLUEPRINTGRAPH_API FMulticastScriptDelegate* GetTargetDelegate() const;
diff --git a/Engine/Source/Editor/BlueprintGraph/Classes/K2Node_CallFunction.h b/Engine/Source/Editor/BlueprintGraph/Classes/K2Node_CallFunction.h
--- a/Engine/Source/Editor/BlueprintGraph/Classes/K2Node_CallFunction.h
+++ b/Engine/Source/Editor/BlueprintGraph/Classes/K2Node_CallFunction.h
@@ -89,6 +89,9 @@ public:
 	virtual void PinDefaultValueChanged(UEdGraphPin* Pin) override;
 	virtual void AddSearchMetaDataInfo(TArray<struct FSearchTagDataPair>& OutTaggedMetaData) const override;
 	virtual TSharedPtr<SWidget> CreateNodeImage() const override;
+	virtual UObject* GetJumpTargetForDoubleClick() const override;
+	virtual bool CanJumpToDefinition() const override;
+	virtual void JumpToDefinition() const override;
 	// End of UEdGraphNode interface
 
 	// UK2Node interface
diff --git a/Engine/Source/Editor/BlueprintGraph/Classes/K2Node_ConstructObjectFromClass.h b/Engine/Source/Editor/BlueprintGraph/Classes/K2Node_ConstructObjectFromClass.h
--- a/Engine/Source/Editor/BlueprintGraph/Classes/K2Node_ConstructObjectFromClass.h
+++ b/Engine/Source/Editor/BlueprintGraph/Classes/K2Node_ConstructObjectFromClass.h
@@ -18,7 +18,6 @@ class BLUEPRINTGRAPH_API UK2Node_ConstructObjectFromClass : public UK2Node
 
 	//~ Begin UEdGraphNode Interface.
 	virtual void AllocateDefaultPins() override;
-	virtual FLinearColor GetNodeTitleColor() const override;
 	virtual FText GetNodeTitle(ENodeTitleType::Type TitleType) const override;
 	virtual void PinDefaultValueChanged(UEdGraphPin* Pin) override;
 	virtual FText GetTooltipText() const override;
@@ -26,6 +25,7 @@ class BLUEPRINTGRAPH_API UK2Node_ConstructObjectFromClass : public UK2Node
 	virtual bool IsCompatibleWithGraph(const UEdGraph* TargetGraph) const override;
 	virtual void PinConnectionListChanged(UEdGraphPin* Pin);
 	virtual void GetPinHoverText(const UEdGraphPin& Pin, FString& HoverTextOut) const override;
+	virtual void PostPlacedNewNode() override;
 	//~ End UEdGraphNode Interface.
 
 	//~ Begin UK2Node Interface
@@ -40,7 +40,7 @@ class BLUEPRINTGRAPH_API UK2Node_ConstructObjectFromClass : public UK2Node
 	void CreatePinsForClass(UClass* InClass, TArray<UEdGraphPin*>* OutClassPins = nullptr);
 
 	/** See if this is a spawn variable pin, or a 'default' pin */
-	virtual bool IsSpawnVarPin(UEdGraphPin* Pin);
+	virtual bool IsSpawnVarPin(UEdGraphPin* Pin) const;
 
 	/** Get the then output pin */
 	UEdGraphPin* GetThenPin() const;
@@ -85,7 +85,6 @@ protected:
 	/** Tooltip text for this node. */
 	FText NodeTooltip;
 
-private:
 	/** Constructing FText strings can be costly, so we cache the node's title */
 	FNodeTextCache CachedNodeTitle;
 };
diff --git a/Engine/Source/Editor/BlueprintGraph/Classes/K2Node_MacroInstance.h b/Engine/Source/Editor/BlueprintGraph/Classes/K2Node_MacroInstance.h
--- a/Engine/Source/Editor/BlueprintGraph/Classes/K2Node_MacroInstance.h
+++ b/Engine/Source/Editor/BlueprintGraph/Classes/K2Node_MacroInstance.h
@@ -55,6 +55,7 @@ public:
 	virtual FString GetDocumentationExcerptName() const override;
 	virtual FSlateIcon GetIconAndTint(FLinearColor& OutColor) const override;
 	virtual bool CanPasteHere(const UEdGraph* TargetGraph) const override;
+	virtual UObject* GetJumpTargetForDoubleClick() const override { return GetMacroGraph(); }
 	//~ End UEdGraphNode Interface
 
 	//~ Begin UK2Node Interface
diff --git a/Engine/Source/Editor/BlueprintGraph/Classes/K2Node_SpawnActorFromClass.h b/Engine/Source/Editor/BlueprintGraph/Classes/K2Node_SpawnActorFromClass.h
--- a/Engine/Source/Editor/BlueprintGraph/Classes/K2Node_SpawnActorFromClass.h
+++ b/Engine/Source/Editor/BlueprintGraph/Classes/K2Node_SpawnActorFromClass.h
@@ -6,6 +6,7 @@
 #include "CoreMinimal.h"
 #include "UObject/ObjectMacros.h"
 #include "K2Node.h"
+#include "K2Node_GenericCreateObject.h"
 #include "Textures/SlateIcon.h"
 #include "EdGraph/EdGraphNodeUtils.h"
 #include "K2Node_SpawnActorFromClass.generated.h"
@@ -14,73 +15,39 @@ class FBlueprintActionDatabaseRegistrar;
 class UEdGraph;
 
 UCLASS()
-class BLUEPRINTGRAPH_API UK2Node_SpawnActorFromClass : public UK2Node
+class BLUEPRINTGRAPH_API UK2Node_SpawnActorFromClass : public UK2Node_ConstructObjectFromClass
 {
 	GENERATED_UCLASS_BODY()
 
 	//~ Begin UEdGraphNode Interface.
 	virtual void AllocateDefaultPins() override;
-	virtual FLinearColor GetNodeTitleColor() const override;
 	virtual FText GetNodeTitle(ENodeTitleType::Type TitleType) const override;
-	virtual void PinDefaultValueChanged(UEdGraphPin* Pin) override;
-	virtual void PinConnectionListChanged(UEdGraphPin* Pin) override;
 	virtual void GetPinHoverText(const UEdGraphPin& Pin, FString& HoverTextOut) const override;
-	virtual FText GetTooltipText() const override;
 	virtual void ExpandNode(class FKismetCompilerContext& CompilerContext, UEdGraph* SourceGraph) override;
 	virtual bool HasExternalDependencies(TArray<class UStruct*>* OptionalOutput) const override;
 	virtual FSlateIcon GetIconAndTint(FLinearColor& OutColor) const override;
 	virtual bool IsCompatibleWithGraph(const UEdGraph* TargetGraph) const override;
-	virtual void PostPlacedNewNode() override;
 	//~ End UEdGraphNode Interface.
 
 	//~ Begin UK2Node Interface
 	virtual bool IsNodeSafeToIgnore() const override { return true; }
 	virtual void ReallocatePinsDuringReconstruction(TArray<UEdGraphPin*>& OldPins) override;
 	virtual void GetNodeAttributes( TArray<TKeyValuePair<FString, FString>>& OutNodeAttributes ) const override;
-	virtual void GetMenuActions(FBlueprintActionDatabaseRegistrar& ActionRegistrar) const override;
-	virtual FText GetMenuCategory() const override;
 	virtual class FNodeHandlingFunctor* CreateNodeHandler(class FKismetCompilerContext& CompilerContext) const override;
 	//~ End UK2Node Interface
 
+	//~ Begin UK2Node_ConstructObjectFromClass Interface
+	virtual UClass* GetClassPinBaseClass() const;
+	virtual bool IsSpawnVarPin(UEdGraphPin* Pin) const override;
+	//~ End UK2Node_ConstructObjectFromClass Interface
 
-	/**
-	 * Create new pins to show properties on archetype
-	 *
-	 * @param InClass		Class to make pins for
-	 * @param OutClassPins	Any pins created through this function will be added to thist list
-	 */
-	void CreatePinsForClass(UClass* InClass, TArray<UEdGraphPin*>& OutClassPins);
-	/** See if this is a spawn variable pin, or a 'default' pin */
-	bool IsSpawnVarPin(UEdGraphPin* Pin);
-
-	/** Get the then output pin */
-	UEdGraphPin* GetThenPin() const;
-	/** Get the blueprint input pin */	
-	UEdGraphPin* GetClassPin(const TArray<UEdGraphPin*>* InPinsToSearch=NULL) const;
-	/** Get the world context input pin, can return NULL */	
-	UEdGraphPin* GetWorldContextPin() const;
+private:
 	/** Get the spawn transform input pin */	
 	UEdGraphPin* GetSpawnTransformPin() const;
 	/** Get the collision handling method input pin */
 	UEdGraphPin* GetCollisionHandlingOverridePin() const;
 	/** Get the actor owner pin */
 	UEdGraphPin* GetOwnerPin() const;
-	/** Get the result output pin */
-	UEdGraphPin* GetResultPin() const;
-
-	/** Get the class that we are going to spawn, if it's defined as default value */
-	UClass* GetClassToSpawn(const TArray<UEdGraphPin*>* InPinsToSearch=NULL) const;
-
-protected:
 
-	void OnClassPinChanged();
-
-	/** Tooltip text for this node. */
-	FText NodeTooltip;
-
-	/** Constructing FText strings can be costly, so we cache the node's title */
-	FNodeTextCache CachedNodeTitle;
-
-private:
 	void MaybeUpdateCollisionPin(TArray<UEdGraphPin*>& OldPins);
 };
diff --git a/Engine/Source/Editor/BlueprintGraph/Classes/K2Node_Timeline.h b/Engine/Source/Editor/BlueprintGraph/Classes/K2Node_Timeline.h
--- a/Engine/Source/Editor/BlueprintGraph/Classes/K2Node_Timeline.h
+++ b/Engine/Source/Editor/BlueprintGraph/Classes/K2Node_Timeline.h
@@ -61,6 +61,7 @@ class UK2Node_Timeline : public UK2Node
 	virtual FString GetDocumentationExcerptName() const override;
 	virtual FSlateIcon GetIconAndTint(FLinearColor& OutColor) const override;
 	virtual bool ShouldShowNodeProperties() const override { return true; }
+	virtual UObject* GetJumpTargetForDoubleClick() const override;
 	//~ End UEdGraphNode Interface.
 
 	//~ Begin UK2Node Interface.
diff --git a/Engine/Source/Editor/BlueprintGraph/Classes/K2Node_Variable.h b/Engine/Source/Editor/BlueprintGraph/Classes/K2Node_Variable.h
--- a/Engine/Source/Editor/BlueprintGraph/Classes/K2Node_Variable.h
+++ b/Engine/Source/Editor/BlueprintGraph/Classes/K2Node_Variable.h
@@ -73,6 +73,9 @@ public:
 	virtual void PostPasteNode() override;
 	virtual bool IsDeprecated() const;
 	virtual FString GetDeprecationMessage() const;
+	virtual UObject* GetJumpTargetForDoubleClick() const override;
+	virtual bool CanJumpToDefinition() const override;
+	virtual void JumpToDefinition() const override;
 	//~ End UEdGraphNode Interface
 
 	//~ Begin K2Node Interface
diff --git a/Engine/Source/Editor/BlueprintGraph/Private/BlueprintEventNodeSpawner.cpp b/Engine/Source/Editor/BlueprintGraph/Private/BlueprintEventNodeSpawner.cpp
--- a/Engine/Source/Editor/BlueprintGraph/Private/BlueprintEventNodeSpawner.cpp
+++ b/Engine/Source/Editor/BlueprintGraph/Private/BlueprintEventNodeSpawner.cpp
@@ -15,19 +15,18 @@
 namespace UBlueprintEventNodeSpawnerImpl
 {
 	/**
-	 * Helper function for removing all disabled nodes connected to a 
-	 * disabled source node
+	 * Helper function for removing a ghost node (and connected ghost nodes, e.g., an autogenerated super call)
 	 *
 	 * @param InNode			The node to start with
 	 * @param InParentGraph		The graph to remove the nodes from
 	 */
-	static void RemoveAllDisabledNodes(UEdGraphNode* InNode, UEdGraph* InParentGraph);
+	static void RemoveGhostNodes(UEdGraphNode* InNode, UEdGraph* InParentGraph);
 }
 
 //------------------------------------------------------------------------------
-static void UBlueprintEventNodeSpawnerImpl::RemoveAllDisabledNodes(UEdGraphNode* InNode, UEdGraph* InParentGraph)
+static void UBlueprintEventNodeSpawnerImpl::RemoveGhostNodes(UEdGraphNode* InNode, UEdGraph* InParentGraph)
 {
-	if(InNode && !InNode->IsNodeEnabled() && !InNode->bUserSetEnabledState)
+	if(InNode && InNode->IsAutomaticallyPlacedGhostNode())
 	{
 		// Go through all pin connections and consume any disabled nodes so we do not leave garbage.
 		for (UEdGraphPin* Pin : InNode->Pins)
@@ -37,7 +36,7 @@ static void UBlueprintEventNodeSpawnerImpl::RemoveAllDisabledNodes(UEdGraphNode*
 			{
 				// Break the pin link back
 				OtherPin->BreakLinkTo(Pin);
-				RemoveAllDisabledNodes(OtherPin->GetOwningNode(), InParentGraph);
+				RemoveGhostNodes(OtherPin->GetOwningNode(), InParentGraph);
 			}
 		}
 
@@ -160,9 +159,9 @@ UEdGraphNode* UBlueprintEventNodeSpawner::Invoke(UEdGraph* ParentGraph, FBinding
 
 	// This Event node might already be present in the Blueprint in a disabled state, 
 	// remove it and allow the user to successfully place the node where they want it.
-	if(EventNode && !EventNode->IsNodeEnabled() && !EventNode->bUserSetEnabledState)
+	if(EventNode && EventNode->IsAutomaticallyPlacedGhostNode())
 	{
-		UBlueprintEventNodeSpawnerImpl::RemoveAllDisabledNodes(EventNode, ParentGraph);
+		UBlueprintEventNodeSpawnerImpl::RemoveGhostNodes(EventNode, ParentGraph);
 		EventNode = nullptr;
 	}
 
diff --git a/Engine/Source/Editor/BlueprintGraph/Private/CallFunctionHandler.cpp b/Engine/Source/Editor/BlueprintGraph/Private/CallFunctionHandler.cpp
--- a/Engine/Source/Editor/BlueprintGraph/Private/CallFunctionHandler.cpp
+++ b/Engine/Source/Editor/BlueprintGraph/Private/CallFunctionHandler.cpp
@@ -367,37 +367,6 @@ void FKCHandler_CallFunction::CreateFunctionCallStatement(FKismetFunctionContext
 			bool bInlineEventCall = false;
 			bool bEmitInstrumentPushState = false;
 			FName EventName = NAME_None;
-			if (Context.IsInstrumentationRequired())
-			{
-				if (UK2Node_CallFunction* CallFunctionNode = Cast<UK2Node_CallFunction>(Node))
-				{
-					const UEdGraphNode* EventNodeOut = nullptr;
-					if (UEdGraph* FunctionGraph = CallFunctionNode->GetFunctionGraph(EventNodeOut))
-					{
-						bEmitInstrumentPushState = true;
-						if (const UK2Node_Event* EventNode = Cast<UK2Node_Event>(EventNodeOut))
-						{
-							bInlineEventCall = true;
-							EventName = EventNode->GetFunctionName();
-						}
-					}
-				}
-			}
-			const bool bInstrumentFunctionEntry = Context.IsInstrumentationRequired() && (bIsLatent || bInlineEventCall || bEmitInstrumentPushState);
-			if (bInstrumentFunctionEntry)
-			{
-				if (bInlineEventCall)
-				{
-					FBlueprintCompiledStatement& EventStatement = Context.AppendStatementForNode(Node);
-					EventStatement.Type = KCST_InstrumentedEvent;
-					EventStatement.Comment = EventName.ToString();
-				}
-				else
-				{
-					FBlueprintCompiledStatement& PushState = Context.AppendStatementForNode(Node);
-					PushState.Type = KCST_InstrumentedStatePush;
-				}
-			}
 
 			// Iterate over all the contexts this functions needs to be called on, and emit a call function statement for each
 			FBlueprintCompiledStatement* LatentStatement = nullptr;
@@ -443,17 +412,6 @@ void FKCHandler_CallFunction::CreateFunctionCallStatement(FKismetFunctionContext
 			// Create the exit from this node if there is one
 			if (bIsLatent)
 			{
-				if (bInstrumentFunctionEntry)
-				{
-					FBlueprintCompiledStatement& SuspendState = Context.AppendStatementForNode(Node);
-					SuspendState.Type = KCST_InstrumentedStateSuspend;
-					if (LatentTargetNode)
-					{
-						check(LatentTargetParamIndex != INDEX_NONE);
-						SuspendState.UbergraphCallIndex = LatentTargetParamIndex;
-						SuspendState.TargetLabel = LatentStatement;
-					}
-				}
 				// End this thread of execution; the latent function will resume it at some point in the future
 				FBlueprintCompiledStatement& PopStatement = Context.AppendStatementForNode(Node);
 				PopStatement.Type = KCST_EndOfThread;
@@ -463,19 +421,6 @@ void FKCHandler_CallFunction::CreateFunctionCallStatement(FKismetFunctionContext
 				// Generate the output impulse from this node
 				if (!IsCalledFunctionPure(Node))
 				{
-					if (bInstrumentFunctionEntry)
-					{
-						if (bInlineEventCall)
-						{
-							FBlueprintCompiledStatement& EventStop = Context.AppendStatementForNode(Node);
-							EventStop.Type = KCST_InstrumentedEventStop;
-						}
-						else
-						{
-							FBlueprintCompiledStatement& PopState = Context.AppendStatementForNode(Node);
-							PopState.Type = KCST_InstrumentedStatePop;
-						}
-					}
 					GenerateSimpleThenGoto(Context, *Node);
 				}
 			}
