Monolithic engine header file exclusion from nativized Blueprint assets is now tied to a project setting.

- Mirrored from //UE4/Dev-Framework (4082035).

diff --git a/Engine/Source/Developer/BlueprintCompilerCppBackend/Private/BlueprintCompilerCppBackendBase.cpp b/Engine/Source/Developer/BlueprintCompilerCppBackend/Private/BlueprintCompilerCppBackendBase.cpp
--- a/Engine/Source/Developer/BlueprintCompilerCppBackend/Private/BlueprintCompilerCppBackendBase.cpp
+++ b/Engine/Source/Developer/BlueprintCompilerCppBackend/Private/BlueprintCompilerCppBackendBase.cpp
@@ -1340,6 +1340,7 @@ void FBlueprintCompilerCppBackendBase::EmitFileBeginning(const FString& CleanNam
 	FIncludeHeaderHelper::EmitIncludeHeader(bIncludeCodeHelpersInHeader ? EmitterContext.Header : EmitterContext.Body, TEXT("GeneratedCodeHelpers"), true);
 	FIncludeHeaderHelper::EmitIncludeHeader(EmitterContext.Header, TEXT("Blueprint/BlueprintSupport"), true);
 
+	FBackendHelperUMG::AdditionalHeaderIncludeForWidget(EmitterContext);
 	FBackendHelperAnim::AddHeaders(EmitterContext);
 
 	TSet<FString> AlreadyIncluded;
diff --git a/Engine/Source/Developer/BlueprintCompilerCppBackend/Private/BlueprintCompilerCppBackendUMG.cpp b/Engine/Source/Developer/BlueprintCompilerCppBackend/Private/BlueprintCompilerCppBackendUMG.cpp
--- a/Engine/Source/Developer/BlueprintCompilerCppBackend/Private/BlueprintCompilerCppBackendUMG.cpp
+++ b/Engine/Source/Developer/BlueprintCompilerCppBackend/Private/BlueprintCompilerCppBackendUMG.cpp
@@ -23,6 +23,15 @@ void FBackendHelperUMG::WidgetFunctionsInHeader(FEmitterLocalContext& Context)
 	}
 }
 
+void FBackendHelperUMG::AdditionalHeaderIncludeForWidget(FEmitterLocalContext& Context)
+{
+	if (!Context.NativizationOptions.bExcludeMonolithicHeaders
+		&& Cast<UWidgetBlueprintGeneratedClass>(Context.GetCurrentlyGeneratedClass()))
+	{
+		Context.Header.AddLine(TEXT("#include \"Runtime/UMG/Public/UMG.h\""));
+	}
+}
+
 void FBackendHelperUMG::CreateClassSubobjects(FEmitterLocalContext& Context, bool bCreate, bool bInitialize)
 {
 	if (auto WidgetClass = Cast<UWidgetBlueprintGeneratedClass>(Context.GetCurrentlyGeneratedClass()))
diff --git a/Engine/Source/Developer/BlueprintCompilerCppBackend/Private/BlueprintCompilerCppBackendUtils.h b/Engine/Source/Developer/BlueprintCompilerCppBackend/Private/BlueprintCompilerCppBackendUtils.h
--- a/Engine/Source/Developer/BlueprintCompilerCppBackend/Private/BlueprintCompilerCppBackendUtils.h
+++ b/Engine/Source/Developer/BlueprintCompilerCppBackend/Private/BlueprintCompilerCppBackendUtils.h
@@ -370,6 +370,8 @@ struct FBackendHelperUMG
 {
 	static void WidgetFunctionsInHeader(FEmitterLocalContext& Context);
 
+	static void AdditionalHeaderIncludeForWidget(FEmitterLocalContext& EmitterContext);
+
 	// these function should use the same context as Constructor
 	static void CreateClassSubobjects(FEmitterLocalContext& Context, bool bCreate, bool bInitialize);
 	static void EmitWidgetInitializationFunctions(FEmitterLocalContext& Context);
diff --git a/Engine/Source/Developer/BlueprintNativeCodeGen/Private/BlueprintNativeCodeGenModule.cpp b/Engine/Source/Developer/BlueprintNativeCodeGen/Private/BlueprintNativeCodeGenModule.cpp
--- a/Engine/Source/Developer/BlueprintNativeCodeGen/Private/BlueprintNativeCodeGenModule.cpp
+++ b/Engine/Source/Developer/BlueprintNativeCodeGen/Private/BlueprintNativeCodeGenModule.cpp
@@ -799,7 +799,7 @@ UObject* FBlueprintNativeCodeGenModule::FindReplacedNameAndOuter(UObject* Object
 
 EReplacementResult FBlueprintNativeCodeGenModule::IsTargetedForReplacement(const UPackage* Package, const FCompilerNativizationOptions& NativizationOptions) const
 {
-	// non-native packages with enums and structs should be converted, unless they are blacklisted:
+	// non-native packages with enums and structs should be converted, unless they are exced:
 	UStruct* Struct = nullptr;
 	UEnum* Enum = nullptr;
 	GetFieldFormPackage(Package, Struct, Enum, RF_NoFlags);
@@ -1219,6 +1219,7 @@ void FBlueprintNativeCodeGenModule::FillPlatformNativizationDetails(const ITarge
 	Details.CompilerNativizationOptions.PlatformName = Details.PlatformName;
 	Details.CompilerNativizationOptions.ClientOnlyPlatform = Platform->IsClientOnly();
 	Details.CompilerNativizationOptions.ServerOnlyPlatform = Platform->IsServerOnly();
+	Details.CompilerNativizationOptions.bExcludeMonolithicHeaders = GetDefault<UProjectPackagingSettings>()->bExcludeMonolithicEngineHeadersInNativizedCode;
 
 	auto GatherExcludedStuff = [&](const TCHAR* KeyForExcludedModules, const TCHAR* KeyForExcludedPaths, const TCHAR* KeyForExcludedAssets)
 	{
diff --git a/Engine/Source/Developer/BlueprintNativeCodeGen/Private/BlueprintNativeCodeGenUtils.cpp b/Engine/Source/Developer/BlueprintNativeCodeGen/Private/BlueprintNativeCodeGenUtils.cpp
--- a/Engine/Source/Developer/BlueprintNativeCodeGen/Private/BlueprintNativeCodeGenUtils.cpp
+++ b/Engine/Source/Developer/BlueprintNativeCodeGen/Private/BlueprintNativeCodeGenUtils.cpp
@@ -30,6 +30,8 @@ DEFINE_LOG_CATEGORY(LogBlueprintCodeGen)
 namespace BlueprintNativeCodeGenUtilsImpl
 {
 	static FString CoreModuleName   = TEXT("Core");
+	static FString EngineModuleName = TEXT("Engine");
+	static FString EngineHeaderFile = TEXT("Engine.h");
 
 	// Used to cache the set of plugin dependencies.
 	static TSet<FString> PluginDependencies;
@@ -48,9 +50,10 @@ namespace BlueprintNativeCodeGenUtilsImpl
 	 * the module to function).
 	 * 
 	 * @param  TargetPaths    Defines the file path/name for the target files.
+	 * @param  bExcludeMonolithicEngineHeaders	Whether or not to exclude monolithic engine headers.
 	 * @return True if the files were successfully generated, otherwise false.
 	 */
-	static bool GenerateModuleSourceFiles(const FBlueprintNativeCodeGenPaths& TargetPaths);
+	static bool GenerateModuleSourceFiles(const FBlueprintNativeCodeGenPaths& TargetPaths, bool bExcludeMonolithicEngineHeaders);
 
 	/**
 	 * Creates and fills out a new .Build.cs file for the plugin's runtime module.
@@ -182,11 +185,15 @@ static bool BlueprintNativeCodeGenUtilsImpl::GeneratePluginDescFile(const FBluep
 }
 
 //------------------------------------------------------------------------------
-static bool BlueprintNativeCodeGenUtilsImpl::GenerateModuleSourceFiles(const FBlueprintNativeCodeGenPaths& TargetPaths)
+static bool BlueprintNativeCodeGenUtilsImpl::GenerateModuleSourceFiles(const FBlueprintNativeCodeGenPaths& TargetPaths, bool bExcludeMonolithicEngineHeaders)
 {
 	FText FailureReason;
 
 	TArray<FString> PchIncludes;
+	if (!bExcludeMonolithicEngineHeaders)
+	{
+		PchIncludes.Add(EngineHeaderFile);
+	}
 	PchIncludes.Add(TEXT("GeneratedCodeHelpers.h"));
 	PchIncludes.Add(NativizedDependenciesFileName() + TEXT(".h"));
 
@@ -256,6 +263,11 @@ static bool BlueprintNativeCodeGenUtilsImpl::GenerateModuleBuildFile(const FBlue
 	TArray<FString> PublicDependencies;
 	// for IModuleInterface
 	PublicDependencies.Add(CoreModuleName);
+	if (!Manifest.GetCompilerNativizationOptions().bExcludeMonolithicHeaders)
+	{
+		// for Engine.h
+		PublicDependencies.Add(EngineModuleName);
+	}
 
 	if (GameProjectUtils::ProjectHasCodeFiles()) 
 	{
@@ -369,7 +381,7 @@ bool FBlueprintNativeCodeGenUtils::FinalizePlugin(const FBlueprintNativeCodeGenM
 	bool bSuccess = true;
 	FBlueprintNativeCodeGenPaths TargetPaths = Manifest.GetTargetPaths();
 	bSuccess = bSuccess && BlueprintNativeCodeGenUtilsImpl::GenerateModuleBuildFile(Manifest);
-	bSuccess = bSuccess && BlueprintNativeCodeGenUtilsImpl::GenerateModuleSourceFiles(TargetPaths);
+	bSuccess = bSuccess && BlueprintNativeCodeGenUtilsImpl::GenerateModuleSourceFiles(TargetPaths, Manifest.GetCompilerNativizationOptions().bExcludeMonolithicHeaders);
 	bSuccess = bSuccess && BlueprintNativeCodeGenUtilsImpl::GenerateNativizedDependenciesSourceFiles(TargetPaths);
 	bSuccess = bSuccess && BlueprintNativeCodeGenUtilsImpl::GeneratePluginDescFile(TargetPaths);
 	return bSuccess;
diff --git a/Engine/Source/Editor/UnrealEd/Classes/Settings/ProjectPackagingSettings.h b/Engine/Source/Editor/UnrealEd/Classes/Settings/ProjectPackagingSettings.h
--- a/Engine/Source/Editor/UnrealEd/Classes/Settings/ProjectPackagingSettings.h
+++ b/Engine/Source/Editor/UnrealEd/Classes/Settings/ProjectPackagingSettings.h
@@ -145,6 +145,10 @@ public:
 	UPROPERTY(config, EditAnywhere, AdvancedDisplay, Category = Blueprints)
 	bool bIncludeNativizedAssetsInProjectGeneration;
 
+	/** Whether or not to exclude monolithic engine headers (e.g. Engine.h) in the generated code when nativizing Blueprint assets. This may improve C++ compiler performance if your game code does not depend on monolithic engine headers to build. */
+	UPROPERTY(config, EditAnywhere, AdvancedDisplay, Category = Blueprints)
+	bool bExcludeMonolithicEngineHeadersInNativizedCode;
+
 	/** If enabled, all content will be put into a single .pak file instead of many individual files (default = enabled). */
 	UPROPERTY(config, EditAnywhere, Category=Packaging)
 	bool UsePakFile;
diff --git a/Engine/Source/Runtime/Engine/Classes/Engine/Blueprint.h b/Engine/Source/Runtime/Engine/Classes/Engine/Blueprint.h
--- a/Engine/Source/Runtime/Engine/Classes/Engine/Blueprint.h
+++ b/Engine/Source/Runtime/Engine/Classes/Engine/Blueprint.h
@@ -101,6 +101,9 @@ struct FCompilerNativizationOptions
 	UPROPERTY()
 	bool ClientOnlyPlatform;
 
+	UPROPERTY()
+	bool bExcludeMonolithicHeaders;
+
 	UPROPERTY()
 	TArray<FName> ExcludedModules;
 
@@ -115,6 +118,7 @@ struct FCompilerNativizationOptions
 	FCompilerNativizationOptions()
 		: ServerOnlyPlatform(false)
 		, ClientOnlyPlatform(false)
+		, bExcludeMonolithicHeaders(false)
 	{}
 };
 
