[NO NON-PII COMMIT MESSAGE]

diff --git a/Engine/Source/Developer/BlueprintNativeCodeGen/Private/BlueprintNativeCodeGenModule.cpp b/Engine/Source/Developer/BlueprintNativeCodeGen/Private/BlueprintNativeCodeGenModule.cpp
--- a/Engine/Source/Developer/BlueprintNativeCodeGen/Private/BlueprintNativeCodeGenModule.cpp
+++ b/Engine/Source/Developer/BlueprintNativeCodeGen/Private/BlueprintNativeCodeGenModule.cpp
@@ -373,11 +373,14 @@ void FBlueprintNativeCodeGenModule::ShutdownModule()
 	// Clear the current coordinator reference.
 	IBlueprintNativeCodeGenCore::Register(nullptr);
 
-	// Reset compiler module delegate function bindings.
-	IBlueprintCompilerCppBackendModule& BackEndModule = (IBlueprintCompilerCppBackendModule&)IBlueprintCompilerCppBackendModule::Get();
-	BackEndModule.GetIsFunctionUsedInADelegateCallback().Unbind();
-	BackEndModule.OnIsTargetedForConversionQuery().Unbind();
-	BackEndModule.OnIncludingUnconvertedBP().Unbind();
+	if (IBlueprintCompilerCppBackendModule::IsAvailable())
+	{
+		// Reset compiler module delegate function bindings.
+		IBlueprintCompilerCppBackendModule& BackEndModule = (IBlueprintCompilerCppBackendModule&)IBlueprintCompilerCppBackendModule::Get();
+		BackEndModule.GetIsFunctionUsedInADelegateCallback().Unbind();
+		BackEndModule.OnIsTargetedForConversionQuery().Unbind();
+		BackEndModule.OnIncludingUnconvertedBP().Unbind();
+	}
 }
 
 void FBlueprintNativeCodeGenModule::GenerateFullyConvertedClasses()
diff --git a/Engine/Source/Developer/GameplayDebugger/Private/GameplayDebuggerModule.cpp b/Engine/Source/Developer/GameplayDebugger/Private/GameplayDebuggerModule.cpp
--- a/Engine/Source/Developer/GameplayDebugger/Private/GameplayDebuggerModule.cpp
+++ b/Engine/Source/Developer/GameplayDebugger/Private/GameplayDebuggerModule.cpp
@@ -88,10 +88,14 @@ void FGameplayDebuggerModule::ShutdownModule()
 	}
 
 #if WITH_EDITOR
-	FPropertyEditorModule& PropertyEditorModule = FModuleManager::GetModuleChecked<FPropertyEditorModule>("PropertyEditor");
-	PropertyEditorModule.UnregisterCustomPropertyTypeLayout("GameplayDebuggerCategoryConfig");
-	PropertyEditorModule.UnregisterCustomPropertyTypeLayout("GameplayDebuggerExtensionConfig");
-	PropertyEditorModule.UnregisterCustomPropertyTypeLayout("GameplayDebuggerInputConfig");
+
+	FPropertyEditorModule* PropertyEditorModule = FModuleManager::GetModulePtr<FPropertyEditorModule>("PropertyEditor");
+	if (PropertyEditorModule)
+	{
+		PropertyEditorModule->UnregisterCustomPropertyTypeLayout("GameplayDebuggerCategoryConfig");
+		PropertyEditorModule->UnregisterCustomPropertyTypeLayout("GameplayDebuggerExtensionConfig");
+		PropertyEditorModule->UnregisterCustomPropertyTypeLayout("GameplayDebuggerInputConfig");
+	}
 
 	FEditorModeRegistry::Get().UnregisterMode(FGameplayDebuggerEdMode::EM_GameplayDebugger);
 #endif
diff --git a/Engine/Source/Developer/MaterialBaking/Private/MaterialBakingModule.cpp b/Engine/Source/Developer/MaterialBaking/Private/MaterialBakingModule.cpp
--- a/Engine/Source/Developer/MaterialBaking/Private/MaterialBakingModule.cpp
+++ b/Engine/Source/Developer/MaterialBaking/Private/MaterialBakingModule.cpp
@@ -72,8 +72,12 @@ void FMaterialBakingModule::StartupModule()
 void FMaterialBakingModule::ShutdownModule()
 {
 	// Urnegister customization and callback
-	FPropertyEditorModule& Module = FModuleManager::Get().LoadModuleChecked<FPropertyEditorModule>("PropertyEditor");
-	Module.UnregisterCustomPropertyTypeLayout(TEXT("PropertyEntry"));
+	FPropertyEditorModule* PropertyEditorModule = FModuleManager::GetModulePtr<FPropertyEditorModule>("PropertyEditor");
+
+	if (PropertyEditorModule)
+	{
+		PropertyEditorModule->UnregisterCustomPropertyTypeLayout(TEXT("PropertyEntry"));
+	}
 	FCoreUObjectDelegates::OnObjectModified.RemoveAll(this);
 }
 
diff --git a/Engine/Source/Editor/AnimationModifiers/Private/AnimationModifiersModule.cpp b/Engine/Source/Editor/AnimationModifiers/Private/AnimationModifiersModule.cpp
--- a/Engine/Source/Editor/AnimationModifiers/Private/AnimationModifiersModule.cpp
+++ b/Engine/Source/Editor/AnimationModifiers/Private/AnimationModifiersModule.cpp
@@ -43,8 +43,11 @@ TSharedRef<FApplicationMode> FAnimationModifiersModule::ExtendApplicationMode(co
 void FAnimationModifiersModule::ShutdownModule()
 {
 	// Make sure we unregister the class layout 
-	FPropertyEditorModule& PropertyEditorModule = FModuleManager::LoadModuleChecked<FPropertyEditorModule>("PropertyEditor");
-	PropertyEditorModule.UnregisterCustomClassLayout("AnimationModifier");
+	FPropertyEditorModule* PropertyEditorModule = FModuleManager::GetModulePtr<FPropertyEditorModule>("PropertyEditor");
+	if (PropertyEditorModule)
+	{
+		PropertyEditorModule->UnregisterCustomClassLayout("AnimationModifier");
+	}
 
 	// Remove extender delegate
 	FWorkflowCentricApplication::GetModeExtenderList().RemoveAll([this](FWorkflowApplicationModeExtender& StoredExtender) { return StoredExtender.GetHandle() == Extender.GetHandle(); });
diff --git a/Engine/Source/Editor/ClothPainter/Private/ClothPaintingModule.cpp b/Engine/Source/Editor/ClothPainter/Private/ClothPaintingModule.cpp
--- a/Engine/Source/Editor/ClothPainter/Private/ClothPaintingModule.cpp
+++ b/Engine/Source/Editor/ClothPainter/Private/ClothPaintingModule.cpp
@@ -195,10 +195,13 @@ void FClothPaintingModule::ShutdownModule()
 	ShutdownMode();
 
 	// Remove skel mesh editor extenders
-	ISkeletalMeshEditorModule& SkelMeshEditorModule = FModuleManager::GetModuleChecked<ISkeletalMeshEditorModule>("SkeletalMeshEditor");
-	TArray<ISkeletalMeshEditorModule::FSkeletalMeshEditorToolbarExtender>& Extenders = SkelMeshEditorModule.GetAllSkeletalMeshEditorToolbarExtenders();
-	
-	Extenders.RemoveAll([=](const ISkeletalMeshEditorModule::FSkeletalMeshEditorToolbarExtender& InDelegate) {return InDelegate.GetHandle() == SkelMeshEditorExtenderHandle; });
+	ISkeletalMeshEditorModule* SkelMeshEditorModule = FModuleManager::GetModulePtr<ISkeletalMeshEditorModule>("SkeletalMeshEditor");
+	if (SkelMeshEditorModule)
+	{
+		TArray<ISkeletalMeshEditorModule::FSkeletalMeshEditorToolbarExtender>& Extenders = SkelMeshEditorModule->GetAllSkeletalMeshEditorToolbarExtenders();
+
+		Extenders.RemoveAll([=](const ISkeletalMeshEditorModule::FSkeletalMeshEditorToolbarExtender& InDelegate) {return InDelegate.GetHandle() == SkelMeshEditorExtenderHandle; });
+	}
 }
 
 void FClothPaintingModule::ShutdownMode()
diff --git a/Engine/Source/Editor/MeshPaintMode/Private/MeshPaintModeModule.cpp b/Engine/Source/Editor/MeshPaintMode/Private/MeshPaintModeModule.cpp
--- a/Engine/Source/Editor/MeshPaintMode/Private/MeshPaintModeModule.cpp
+++ b/Engine/Source/Editor/MeshPaintMode/Private/MeshPaintModeModule.cpp
@@ -40,9 +40,12 @@ void FMeshPaintModeModule::ShutdownModule()
 	FEditorModeRegistry::Get().UnregisterMode(FBuiltinEditorModes::EM_MeshPaint);
 
 	/** De-register detail/property customization */
-	FPropertyEditorModule& PropertyModule = FModuleManager::LoadModuleChecked<FPropertyEditorModule>("PropertyEditor");
-	PropertyModule.UnregisterCustomClassLayout("PaintModeSettings");
-	PropertyModule.UnregisterCustomPropertyTypeLayout("VertexPaintSettings");
-	PropertyModule.UnregisterCustomPropertyTypeLayout("TexturePaintSettings");
-	PropertyModule.UnregisterCustomPropertyTypeLayout("TexturePaintSettings");
+	FPropertyEditorModule* PropertyModule = FModuleManager::GetModulePtr<FPropertyEditorModule>("PropertyEditor");
+	if (PropertyModule)
+	{
+		PropertyModule->UnregisterCustomClassLayout("PaintModeSettings");
+		PropertyModule->UnregisterCustomPropertyTypeLayout("VertexPaintSettings");
+		PropertyModule->UnregisterCustomPropertyTypeLayout("TexturePaintSettings");
+		PropertyModule->UnregisterCustomPropertyTypeLayout("TexturePaintSettings");
+	}
 }
diff --git a/Engine/Source/Editor/TranslationEditor/Private/TranslationEditorModule.cpp b/Engine/Source/Editor/TranslationEditor/Private/TranslationEditorModule.cpp
--- a/Engine/Source/Editor/TranslationEditor/Private/TranslationEditorModule.cpp
+++ b/Engine/Source/Editor/TranslationEditor/Private/TranslationEditorModule.cpp
@@ -36,8 +36,11 @@ void FTranslationEditorModule::ShutdownModule()
 
 #if WITH_UNREAL_DEVELOPER_TOOLS
 	// unregister message log
-	FMessageLogModule& MessageLogModule = FModuleManager::LoadModuleChecked<FMessageLogModule>("MessageLog");
-	MessageLogModule.UnregisterLogListing("TranslationEditor");
+	FMessageLogModule* MessageLogModule = FModuleManager::LoadModulePtr<FMessageLogModule>("MessageLog");
+	if (MessageLogModule)
+	{
+		MessageLogModule->UnregisterLogListing("TranslationEditor");
+	}
 #endif
 }
 
