Fix UBT warnings when building/packaging with Blueprint nativization enabled.

Change summary:
- Modified the PluginModule Build.cs template to include the explicit 'PrivatePCHHeaderFile' assignment.
- Added BlueprintSupport.h to the NativizedAssets plugin module's PCH file.
- Modified generated header/source files to be more IWYU-compliant when monolithic header files are excluded.
- Modified UEBuildTarget.CreateTarget() to validate a nativized build configuration only for game-specific target types.

diff --git a/Engine/Source/Developer/BlueprintCompilerCppBackend/Private/BlueprintCompilerCppBackendBase.cpp b/Engine/Source/Developer/BlueprintCompilerCppBackend/Private/BlueprintCompilerCppBackendBase.cpp
--- a/Engine/Source/Developer/BlueprintCompilerCppBackend/Private/BlueprintCompilerCppBackendBase.cpp
+++ b/Engine/Source/Developer/BlueprintCompilerCppBackend/Private/BlueprintCompilerCppBackendBase.cpp
@@ -986,17 +986,21 @@ void FBlueprintCompilerCppBackendBase::GenerateCodeFromEnum(UUserDefinedEnum* So
 	FCodeText Body;
 	
 	const FString PCHFilename = FEmitHelper::GetPCHFilename();
-	if (!PCHFilename.IsEmpty())
+	if (!NativizationOptions.bExcludeMonolithicHeaders)
 	{
-		Body.AddLine(FString::Printf(TEXT("#include \"%s\""), *PCHFilename));	
-	}
-	else
-	{
-		// Used when generated code is not in a separate module
-		const FString MainHeaderFilename = FEmitHelper::GetGameMainHeaderFilename();
-		if (!MainHeaderFilename.IsEmpty())
+		// If monolithic headers are not being excluded, then we are not IWYU-compliant, and thus we include the PCH file first.
+		if (!PCHFilename.IsEmpty())
+		{
+			Body.AddLine(FString::Printf(TEXT("#include \"%s\""), *PCHFilename));
+		}
+		else
 		{
-			Body.AddLine(FString::Printf(TEXT("#include \"%s\""), *MainHeaderFilename));
+			// Used when generated code is not in a separate module
+			const FString MainHeaderFilename = FEmitHelper::GetGameMainHeaderFilename();
+			if (!MainHeaderFilename.IsEmpty())
+			{
+				Body.AddLine(FString::Printf(TEXT("#include \"%s\""), *MainHeaderFilename));
+			}
 		}
 	}
 
@@ -1362,17 +1366,21 @@ void FBlueprintCompilerCppBackendBase::EmitFileBeginning(const FString& CleanNam
 	EmitterContext.Header.AddLine(TEXT("#pragma once"));
 
 	const FString PCHFilename = FEmitHelper::GetPCHFilename();
-	if (!PCHFilename.IsEmpty())
+	if (!EmitterContext.NativizationOptions.bExcludeMonolithicHeaders)
 	{
-		FIncludeHeaderHelper::EmitIncludeHeader(EmitterContext.Body, *PCHFilename, false);
-	}
-	else
-	{
-		// Used when generated code is not in a separate module
-		const FString MainHeaderFilename = FEmitHelper::GetGameMainHeaderFilename();
-		if (!MainHeaderFilename.IsEmpty())
+		// If monolithic headers are not being excluded, then we are not IWYU-compliant, and thus we include the PCH file first.
+		if (!PCHFilename.IsEmpty())
+		{
+			FIncludeHeaderHelper::EmitIncludeHeader(EmitterContext.Body, *PCHFilename, false);
+		}
+		else
 		{
-			FIncludeHeaderHelper::EmitIncludeHeader(EmitterContext.Body, *MainHeaderFilename, false);
+			// Used when generated code is not in a separate module
+			const FString MainHeaderFilename = FEmitHelper::GetGameMainHeaderFilename();
+			if (!MainHeaderFilename.IsEmpty())
+			{
+				FIncludeHeaderHelper::EmitIncludeHeader(EmitterContext.Body, *MainHeaderFilename, false);
+			}
 		}
 	}
 	
diff --git a/Engine/Source/Developer/BlueprintNativeCodeGen/Private/BlueprintNativeCodeGenUtils.cpp b/Engine/Source/Developer/BlueprintNativeCodeGen/Private/BlueprintNativeCodeGenUtils.cpp
--- a/Engine/Source/Developer/BlueprintNativeCodeGen/Private/BlueprintNativeCodeGenUtils.cpp
+++ b/Engine/Source/Developer/BlueprintNativeCodeGen/Private/BlueprintNativeCodeGenUtils.cpp
@@ -75,7 +75,7 @@ namespace BlueprintNativeCodeGenUtilsImpl
 	/** */
 	static FString NativizedDependenciesFileName() { return TEXT("NativizedAssets_Dependencies"); }
 	/** */
-	static bool GenerateNativizedDependenciesSourceFiles(const FBlueprintNativeCodeGenPaths& TargetPaths);
+	static bool GenerateNativizedDependenciesSourceFiles(const FBlueprintNativeCodeGenPaths& TargetPaths, bool bExcludeMonolithicEngineHeaders);
 }
 
 //------------------------------------------------------------------------------
@@ -195,6 +195,7 @@ static bool BlueprintNativeCodeGenUtilsImpl::GenerateModuleSourceFiles(const FBl
 		PchIncludes.Add(EngineHeaderFile);
 	}
 	PchIncludes.Add(TEXT("GeneratedCodeHelpers.h"));
+	PchIncludes.Add(TEXT("Blueprint/BlueprintSupport.h"));
 	PchIncludes.Add(NativizedDependenciesFileName() + TEXT(".h"));
 
 	TArray<FString> FilesToIncludeInModuleHeader;
@@ -218,7 +219,7 @@ static bool BlueprintNativeCodeGenUtilsImpl::GenerateModuleSourceFiles(const FBl
 }
 
 //------------------------------------------------------------------------------
-static bool BlueprintNativeCodeGenUtilsImpl::GenerateNativizedDependenciesSourceFiles(const FBlueprintNativeCodeGenPaths& TargetPaths)
+static bool BlueprintNativeCodeGenUtilsImpl::GenerateNativizedDependenciesSourceFiles(const FBlueprintNativeCodeGenPaths& TargetPaths, bool bExcludeMonolithicEngineHeaders)
 {
 	FText FailureReason;
 	bool bSuccess = true;
@@ -234,7 +235,7 @@ static bool BlueprintNativeCodeGenUtilsImpl::GenerateNativizedDependenciesSource
 
 	{
 		const FString SourceFilePath = FPaths::Combine(*TargetPaths.RuntimeSourceDir(FBlueprintNativeCodeGenPaths::CppFile), *BaseFilename) + TEXT(".cpp");
-		const FString SourceFileContent = CodeGenBackend.DependenciesGlobalMapBodyCode(TargetPaths.RuntimeModuleName());
+		const FString SourceFileContent = CodeGenBackend.DependenciesGlobalMapBodyCode(bExcludeMonolithicEngineHeaders ? BaseFilename : TargetPaths.RuntimeModuleName());
 		bSuccess &= GameProjectUtils::WriteOutputFile(SourceFilePath, SourceFileContent, FailureReason);
 	}
 
@@ -337,7 +338,7 @@ static bool BlueprintNativeCodeGenUtilsImpl::GenerateModuleBuildFile(const FBlue
 	
 	FText ErrorMessage;
 	bool bSuccess = GameProjectUtils::GeneratePluginModuleBuildFile(TargetPaths.RuntimeBuildFile(), TargetPaths.RuntimeModuleName(),
-		PublicDependencies, PrivateDependencies, ErrorMessage, false);
+		PublicDependencies, PrivateDependencies, ErrorMessage, Manifest.GetCompilerNativizationOptions().bExcludeMonolithicHeaders);
 
 	if (!bSuccess)
 	{
@@ -379,10 +380,11 @@ static UClass* BlueprintNativeCodeGenUtilsImpl::ResolveReplacementType(const FCo
 bool FBlueprintNativeCodeGenUtils::FinalizePlugin(const FBlueprintNativeCodeGenManifest& Manifest)
 {
 	bool bSuccess = true;
+	const bool bExcludeMonolithicHeaders = Manifest.GetCompilerNativizationOptions().bExcludeMonolithicHeaders;
 	FBlueprintNativeCodeGenPaths TargetPaths = Manifest.GetTargetPaths();
 	bSuccess = bSuccess && BlueprintNativeCodeGenUtilsImpl::GenerateModuleBuildFile(Manifest);
-	bSuccess = bSuccess && BlueprintNativeCodeGenUtilsImpl::GenerateModuleSourceFiles(TargetPaths, Manifest.GetCompilerNativizationOptions().bExcludeMonolithicHeaders);
-	bSuccess = bSuccess && BlueprintNativeCodeGenUtilsImpl::GenerateNativizedDependenciesSourceFiles(TargetPaths);
+	bSuccess = bSuccess && BlueprintNativeCodeGenUtilsImpl::GenerateModuleSourceFiles(TargetPaths, bExcludeMonolithicHeaders);
+	bSuccess = bSuccess && BlueprintNativeCodeGenUtilsImpl::GenerateNativizedDependenciesSourceFiles(TargetPaths, bExcludeMonolithicHeaders);
 	bSuccess = bSuccess && BlueprintNativeCodeGenUtilsImpl::GeneratePluginDescFile(TargetPaths);
 	return bSuccess;
 }
diff --git a/Engine/Source/Programs/UnrealBuildTool/Configuration/UEBuildTarget.cs b/Engine/Source/Programs/UnrealBuildTool/Configuration/UEBuildTarget.cs
--- a/Engine/Source/Programs/UnrealBuildTool/Configuration/UEBuildTarget.cs
+++ b/Engine/Source/Programs/UnrealBuildTool/Configuration/UEBuildTarget.cs
@@ -469,7 +469,9 @@ namespace UnrealBuildTool
 			}
 
 			// Include generated code plugin if not building an editor target and project is configured for nativization
-			if (RulesObject.ProjectFile != null && RulesObject.Type != TargetType.Editor && ShouldIncludeNativizedAssets(RulesObject.ProjectFile.Directory))
+			if (RulesObject.ProjectFile != null
+				&& (RulesObject.Type == TargetType.Game || RulesObject.Type == TargetType.Client || RulesObject.Type == TargetType.Server)
+				&& ShouldIncludeNativizedAssets(RulesObject.ProjectFile.Directory))
 			{
 				string PlatformName;
 				if (RulesObject.Platform == UnrealTargetPlatform.Win32 || RulesObject.Platform == UnrealTargetPlatform.Win64)
