PR: Only allow BPGCs to be excluded from conversion by folder path (for consistency with initial design/comment).

#5124

diff --git a/Engine/Source/Developer/BlueprintNativeCodeGen/Private/BlueprintNativeCodeGenModule.cpp b/Engine/Source/Developer/BlueprintNativeCodeGen/Private/BlueprintNativeCodeGenModule.cpp
--- a/Engine/Source/Developer/BlueprintNativeCodeGen/Private/BlueprintNativeCodeGenModule.cpp
+++ b/Engine/Source/Developer/BlueprintNativeCodeGen/Private/BlueprintNativeCodeGenModule.cpp
@@ -900,7 +900,8 @@ EReplacementResult FBlueprintNativeCodeGenModule::IsTargetedForReplacement(const
 
 	auto ObjectGeneratesOnlyStub = [&]() -> bool
 	{
-		// ExcludedFolderPaths
+		// ExcludedFolderPaths - Only BPGCs are excluded by path.
+		if (BlueprintClass)
 		{
 			const FString ObjPathName = Object->GetPathName();
 			for (const FString& ExcludedPath : ExcludedFolderPaths)
