Fix circular dependency issues between nativized BPGC/UDS and non-nativized Blueprint class assets causing EDL cycles at runtime in exclusively-nativized builds.

Change summary:
- Modified FBlueprintNativeCodeGenModule::ReplaceAssetsWithCircularReferenceTo() and HasCircularReferenceWithAnyConvertedAsset() utility methods to include UDS types in circular reference checks. This ensures that a non-nativized Blueprint class asset (e.g. DOBP) that shares a circular reference with a UDS type (which is always nativized) will force the non-nativized asset to be converted as well. UDS assets are now consistent with BPGC assets in this regard.

diff --git a/Engine/Source/Developer/BlueprintNativeCodeGen/Private/BlueprintNativeCodeGenModule.cpp b/Engine/Source/Developer/BlueprintNativeCodeGen/Private/BlueprintNativeCodeGenModule.cpp
--- a/Engine/Source/Developer/BlueprintNativeCodeGen/Private/BlueprintNativeCodeGenModule.cpp
+++ b/Engine/Source/Developer/BlueprintNativeCodeGen/Private/BlueprintNativeCodeGenModule.cpp
@@ -80,9 +80,9 @@ private:
 	void CollectBoundFunctions(UBlueprint* BP);
 	void GenerateSingleAsset(UField* ForConversion, const FName PlatformName, TSharedPtr<FNativizationSummary> NativizationSummary = TSharedPtr<FNativizationSummary>());
 	void ReplaceAsset(const UObject* InAsset, const FCompilerNativizationOptions& NativizationOptions) const;
-	void GatherClassAssetsReferencedByStruct(TSet<const UBlueprintGeneratedClass*>& Assets, const UStruct* OuterStruct, const UStruct* InnerStruct = nullptr) const;
-	void ReplaceAssetsWithCircularReferenceTo(const UBlueprintGeneratedClass* InClass, const FCompilerNativizationOptions& NativizationOptions) const;
-	bool HasCircularReferenceWithAnyConvertedAsset(const UBlueprintGeneratedClass* InClass, const FCompilerNativizationOptions& NativizationOptions) const;
+	void GatherConvertableAssetsReferencedByStruct(TSet<const UStruct*>& Assets, const UStruct* OuterStruct, const UStruct* InnerStruct = nullptr) const;
+	void ReplaceAssetsWithCircularReferenceTo(const UStruct* InStruct, const FCompilerNativizationOptions& NativizationOptions) const;
+	bool HasCircularReferenceWithAnyConvertedAsset(const UStruct* InStruct, const FCompilerNativizationOptions& NativizationOptions) const;
 
 	struct FStatePerPlatform
 	{
@@ -1056,19 +1056,22 @@ EReplacementResult FBlueprintNativeCodeGenModule::IsTargetedForReplacement(const
 
 	StateForCurrentPlatform->CachedIsTargetedForReplacement.Add(ObjectKey, Result);
 
-	if (BlueprintClass)
+	if (const UStruct* ObjAsStruct = Cast<UStruct>(Object))
 	{
 		if (Result == EReplacementResult::ReplaceCompletely)
 		{
 			// Look for any circular references with unconverted assets. We'll need to convert those as well in order to avoid creating an EDL cycle.
-			ReplaceAssetsWithCircularReferenceTo(BlueprintClass, NativizationOptions);
+			ReplaceAssetsWithCircularReferenceTo(ObjAsStruct, NativizationOptions);
 		}
-		else if(HasCircularReferenceWithAnyConvertedAsset(BlueprintClass, NativizationOptions))
+		else if(HasCircularReferenceWithAnyConvertedAsset(ObjAsStruct, NativizationOptions))
 		{
-			UE_LOG(LogBlueprintCodeGen, Log, TEXT("Forcing '%s' to be replaced as it has a circular reference to a converted asset"), *BlueprintClass->GetName());
+			UE_LOG(LogBlueprintCodeGen, Log, TEXT("Forcing '%s' to be replaced as it has a circular reference to a converted asset"), *ObjAsStruct->GetName());
 
 			// Force unconverted assets to be replaced if it has a circular reference with any converted asset.
-			ReplaceAsset(BlueprintClass, NativizationOptions);
+			ReplaceAsset(ObjAsStruct, NativizationOptions);
+
+			// Update the result.
+			Result = StateForCurrentPlatform->CachedIsTargetedForReplacement.FindChecked(ObjectKey);
 		}
 	}
 
@@ -1087,7 +1090,7 @@ void FBlueprintNativeCodeGenModule::ReplaceAsset(const UObject* InAsset, const F
 	}
 }
 
-void FBlueprintNativeCodeGenModule::GatherClassAssetsReferencedByStruct(TSet<const UBlueprintGeneratedClass*>& Assets, const UStruct* OuterStruct, const UStruct* InnerStruct) const
+void FBlueprintNativeCodeGenModule::GatherConvertableAssetsReferencedByStruct(TSet<const UStruct*>& Assets, const UStruct* OuterStruct, const UStruct* InnerStruct) const
 {
 	if (OuterStruct)
 	{
@@ -1116,18 +1119,26 @@ void FBlueprintNativeCodeGenModule::GatherClassAssetsReferencedByStruct(TSet<con
 			{
 				if (const UStructProperty* StructProperty = Cast<UStructProperty>(InnerProperty))
 				{
-					GatherClassAssetsReferencedByStruct(Assets, OuterStruct, StructProperty->Struct);
+					if (const UUserDefinedStruct* UDS = Cast<UUserDefinedStruct>(StructProperty->Struct))
+					{
+						Assets.Add(StructProperty->Struct);
+					}
+
+					GatherConvertableAssetsReferencedByStruct(Assets, OuterStruct, StructProperty->Struct);
 				}
 				else
 				{
 					const UBlueprintGeneratedClass* BPGC = nullptr;
 					if (const UObjectPropertyBase* ObjectProperty = Cast<UObjectPropertyBase>(InnerProperty))
 					{
-						BPGC = Cast<const UBlueprintGeneratedClass>(ObjectProperty->PropertyClass);
-					}
-					else if (const UClassProperty* ClassProperty = Cast<UClassProperty>(InnerProperty))
-					{
-						BPGC = Cast<const UBlueprintGeneratedClass>(ClassProperty->MetaClass);
+						if (const UClassProperty* ClassProperty = Cast<UClassProperty>(InnerProperty))
+						{
+							BPGC = Cast<const UBlueprintGeneratedClass>(ClassProperty->MetaClass);
+						}
+						else
+						{
+							BPGC = Cast<const UBlueprintGeneratedClass>(ObjectProperty->PropertyClass);
+						}
 					}
 
 					if (BPGC)
@@ -1140,33 +1151,33 @@ void FBlueprintNativeCodeGenModule::GatherClassAssetsReferencedByStruct(TSet<con
 	}
 }
 
-void FBlueprintNativeCodeGenModule::ReplaceAssetsWithCircularReferenceTo(const UBlueprintGeneratedClass* InClass, const FCompilerNativizationOptions& NativizationOptions) const
+void FBlueprintNativeCodeGenModule::ReplaceAssetsWithCircularReferenceTo(const UStruct* InStruct, const FCompilerNativizationOptions& NativizationOptions) const
 {
-	TSet<const UBlueprintGeneratedClass*> ForwardReferencedAssets;
-	GatherClassAssetsReferencedByStruct(ForwardReferencedAssets, InClass);
+	TSet<const UStruct*> ForwardReferencedAssets;
+	GatherConvertableAssetsReferencedByStruct(ForwardReferencedAssets, InStruct);
 
-	for (const UBlueprintGeneratedClass* ForwardReference : ForwardReferencedAssets)
+	for (const UStruct* ForwardReference : ForwardReferencedAssets)
 	{
 		const EReplacementResult Result = IsTargetedForReplacement(ForwardReference, NativizationOptions);
 		if (Result != EReplacementResult::ReplaceCompletely)
 		{
 			bool bForceConvert = false;
-			if (ForwardReference->IsChildOf(InClass))
+			if (ForwardReference->IsChildOf(InStruct))
 			{
-				UE_LOG(LogBlueprintCodeGen, Log, TEXT("Forcing '%s' to be replaced as it has a circular reference with '%s'"), *ForwardReference->GetName(), *InClass->GetName());
+				UE_LOG(LogBlueprintCodeGen, Log, TEXT("Forcing '%s' to be replaced as it has a circular reference with '%s'"), *ForwardReference->GetName(), *InStruct->GetName());
 				
 				ReplaceAsset(ForwardReference, NativizationOptions);
 			}
 			else
 			{
-				TSet<const UBlueprintGeneratedClass*> ReverseReferencedAssets;
-				GatherClassAssetsReferencedByStruct(ReverseReferencedAssets, ForwardReference);
+				TSet<const UStruct*> ReverseReferencedAssets;
+				GatherConvertableAssetsReferencedByStruct(ReverseReferencedAssets, ForwardReference);
 
-				for (const UBlueprintGeneratedClass* ReverseReference : ReverseReferencedAssets)
+				for (const UStruct* ReverseReference : ReverseReferencedAssets)
 				{
-					if (ReverseReference->IsChildOf(InClass))
+					if (ReverseReference->IsChildOf(InStruct))
 					{
-						UE_LOG(LogBlueprintCodeGen, Log, TEXT("Forcing '%s' to be replaced as it has a circular reference to '%s'"), *ForwardReference->GetName(), *InClass->GetName());
+						UE_LOG(LogBlueprintCodeGen, Log, TEXT("Forcing '%s' to be replaced as it has a circular reference to '%s'"), *ForwardReference->GetName(), *InStruct->GetName());
 						
 						ReplaceAsset(ForwardReference, NativizationOptions);
 						break;
@@ -1177,22 +1188,22 @@ void FBlueprintNativeCodeGenModule::ReplaceAssetsWithCircularReferenceTo(const U
 	}
 }
 
-bool FBlueprintNativeCodeGenModule::HasCircularReferenceWithAnyConvertedAsset(const UBlueprintGeneratedClass* InClass, const FCompilerNativizationOptions& NativizationOptions) const
+bool FBlueprintNativeCodeGenModule::HasCircularReferenceWithAnyConvertedAsset(const UStruct* InStruct, const FCompilerNativizationOptions& NativizationOptions) const
 {
-	TSet<const UBlueprintGeneratedClass*> ForwardReferencedAssets;
-	GatherClassAssetsReferencedByStruct(ForwardReferencedAssets, InClass);
+	TSet<const UStruct*> ForwardReferencedAssets;
+	GatherConvertableAssetsReferencedByStruct(ForwardReferencedAssets, InStruct);
 
-	for (const UBlueprintGeneratedClass* ForwardReference : ForwardReferencedAssets)
+	for (const UStruct* ForwardReference : ForwardReferencedAssets)
 	{
 		const EReplacementResult Result = IsTargetedForReplacement(ForwardReference, NativizationOptions);
 		if (Result == EReplacementResult::ReplaceCompletely)
 		{
-			TSet<const UBlueprintGeneratedClass*> ReverseReferencedAssets;
-			GatherClassAssetsReferencedByStruct(ReverseReferencedAssets, ForwardReference);
+			TSet<const UStruct*> ReverseReferencedAssets;
+			GatherConvertableAssetsReferencedByStruct(ReverseReferencedAssets, ForwardReference);
 
-			for (const UBlueprintGeneratedClass* ReverseReference : ReverseReferencedAssets)
+			for (const UStruct* ReverseReference : ReverseReferencedAssets)
 			{
-				if (ReverseReference->IsChildOf(InClass))
+				if (ReverseReference->IsChildOf(InStruct))
 				{
 					return true;
 				}
