Exclude class assets that inherit from a native C++ class with sparse data from Blueprint nativization. It's not supported at this time.

#rnx

diff --git a/Engine/Source/Developer/BlueprintNativeCodeGen/Private/BlueprintNativeCodeGenModule.cpp b/Engine/Source/Developer/BlueprintNativeCodeGen/Private/BlueprintNativeCodeGenModule.cpp
--- a/Engine/Source/Developer/BlueprintNativeCodeGen/Private/BlueprintNativeCodeGenModule.cpp
+++ b/Engine/Source/Developer/BlueprintNativeCodeGen/Private/BlueprintNativeCodeGenModule.cpp
@@ -1039,6 +1039,28 @@ EReplacementResult FBlueprintNativeCodeGenModule::IsTargetedForReplacement(const
 					return true;
 				}
 			}
+
+			// @todo: Remove once we've added support for nativizing a Blueprint that inherits from a native class that's been associated with sparse class data.
+			if (BlueprintClass->GetSparseClassDataStruct())
+			{
+				UE_LOG(LogBlueprintCodeGen, Warning, TEXT("BP %s is selected for nativization, but it cannot be nativized because it's based on a C++ class hierarchy that includes sparse class data. Nativization of this BP asset type is not supported at this time."), *GetPathNameSafe(Blueprint));
+				return true;
+			}
+
+			// @todo: Remove once we've added support for nativizing a Blueprint that references another Blueprint that inherits from a native class that's been associated with sparse class data.
+			TSet<const UStruct*> ReferencedAssets;
+			GatherConvertableAssetsReferencedByStruct(ReferencedAssets, BlueprintClass);
+			for (const UStruct* ReferencedAsset : ReferencedAssets)
+			{
+				if (const UBlueprintGeneratedClass* BPGC = Cast<UBlueprintGeneratedClass>(ReferencedAsset))
+				{
+					if (BPGC->GetSparseClassDataStruct())
+					{
+						UE_LOG(LogBlueprintCodeGen, Warning, TEXT("BP %s is selected for nativization, but it cannot be nativized because it currently references another BP (%s) that's based on a C++ class hierarchy that includes sparse class data. Nativization of this BP asset type is not supported at this time."), *GetPathNameSafe(Blueprint), *GetPathNameSafe(BPGC->GetOuter()));
+						return true;
+					}
+				}
+			}
 		}
 		return false;
 	};
