Fix for build exception during nativization of Match3 sample project (regression).

#rnx

diff --git a/Engine/Source/Developer/BlueprintNativeCodeGen/Private/BlueprintNativeCodeGenManifest.cpp b/Engine/Source/Developer/BlueprintNativeCodeGen/Private/BlueprintNativeCodeGenManifest.cpp
--- a/Engine/Source/Developer/BlueprintNativeCodeGen/Private/BlueprintNativeCodeGenManifest.cpp
+++ b/Engine/Source/Developer/BlueprintNativeCodeGen/Private/BlueprintNativeCodeGenManifest.cpp
@@ -236,7 +236,7 @@ static bool BlueprintNativeCodeGenManifestImpl::GatherModuleDependencies(const U
 			}
 
 			// we want only native packages, ones that are not editor-only
-			if ((DependentPackage->GetPackageFlags() & (PKG_CompiledIn | PKG_EditorOnly | PKG_Developer)) == PKG_CompiledIn)
+			if ((DependentPackage->HasAnyPackageFlags(PKG_CompiledIn)) && !DependentPackage->HasAnyPackageFlags(PKG_EditorOnly | PKG_Developer) && !DependentPackage->IsLoadedByEditorPropertiesOnly())
 			{
 				DependenciesOut.AddUnique(DependentPackage);// PkgImport.ObjectName.ToString());
 			}
