Fix incorrect target platform names for some client/server configs that led to an ensure at cook time with Blueprint nativization enabled.

- Mirrored from //UE4/Dev-Framework (CL# 10940643).

diff --git a/Engine/Source/Developer/BlueprintNativeCodeGen/Private/BlueprintNativeCodeGenModule.cpp b/Engine/Source/Developer/BlueprintNativeCodeGen/Private/BlueprintNativeCodeGenModule.cpp
--- a/Engine/Source/Developer/BlueprintNativeCodeGen/Private/BlueprintNativeCodeGenModule.cpp
+++ b/Engine/Source/Developer/BlueprintNativeCodeGen/Private/BlueprintNativeCodeGenModule.cpp
@@ -1229,9 +1229,8 @@ bool FBlueprintNativeCodeGenModule::HasCircularReferenceWithAnyConvertedAsset(co
 void FBlueprintNativeCodeGenModule::FillPlatformNativizationDetails(const ITargetPlatform* Platform, FPlatformNativizationDetails& Details)
 {
 	check(Platform);
-	const PlatformInfo::FPlatformInfo& PlatformInfo = Platform->GetPlatformInfo();
 
-	Details.PlatformName = PlatformInfo.TargetPlatformName;
+	Details.PlatformName = FName(*Platform->PlatformName());
 	Details.CompilerNativizationOptions.PlatformName = Details.PlatformName;
 	Details.CompilerNativizationOptions.ClientOnlyPlatform = Platform->IsClientOnly();
 	Details.CompilerNativizationOptions.ServerOnlyPlatform = Platform->IsServerOnly();
