Fix incorrect platform name being used in manifest lookup for some cook targets during Blueprint nativization.

- Mirrored from //UE4/Dev-Framework (CL 11060400).

diff --git a/Engine/Source/Developer/BlueprintNativeCodeGen/Private/BlueprintNativeCodeGenModule.cpp b/Engine/Source/Developer/BlueprintNativeCodeGen/Private/BlueprintNativeCodeGenModule.cpp
--- a/Engine/Source/Developer/BlueprintNativeCodeGen/Private/BlueprintNativeCodeGenModule.cpp
+++ b/Engine/Source/Developer/BlueprintNativeCodeGen/Private/BlueprintNativeCodeGenModule.cpp
@@ -116,7 +116,7 @@ private:
 
 const FCompilerNativizationOptions& FBlueprintNativeCodeGenModule::GetNativizationOptionsForPlatform(const ITargetPlatform* Platform) const
 {
-	const FName PlatformName = ensure(Platform) ? Platform->GetPlatformInfo().PlatformInfoName : NAME_None;
+	const FName PlatformName = ensure(Platform) ? FName(*Platform->PlatformName()) : NAME_None;
 
 	const TUniquePtr<FBlueprintNativeCodeGenManifest>* Result = Manifests.Find(PlatformName);
 	if (ensure(Result && Result->IsValid()))
