Extend the standalone BP class nativization tool to interface types.

diff --git a/Engine/Source/Developer/BlueprintNativeCodeGen/Private/NativeCodeGenerationTool.cpp b/Engine/Source/Developer/BlueprintNativeCodeGen/Private/NativeCodeGenerationTool.cpp
--- a/Engine/Source/Developer/BlueprintNativeCodeGen/Private/NativeCodeGenerationTool.cpp
+++ b/Engine/Source/Developer/BlueprintNativeCodeGen/Private/NativeCodeGenerationTool.cpp
@@ -378,7 +378,7 @@ void FNativeCodeGenerationTool::Open(UBlueprint& Blueprint, TSharedRef< class FB
 bool FNativeCodeGenerationTool::CanGenerate(const UBlueprint& Blueprint)
 {
 	return (Blueprint.Status == EBlueprintStatus::BS_UpToDate || Blueprint.Status == EBlueprintStatus::BS_UpToDateWithWarnings)
-		&& (Blueprint.BlueprintType == EBlueprintType::BPTYPE_Normal || Blueprint.BlueprintType == EBlueprintType::BPTYPE_FunctionLibrary)
+		&& (Blueprint.BlueprintType == EBlueprintType::BPTYPE_Normal || Blueprint.BlueprintType == EBlueprintType::BPTYPE_FunctionLibrary || Blueprint.BlueprintType == EBlueprintType::BPTYPE_Interface)
 		&& Cast<UBlueprintGeneratedClass>(Blueprint.GeneratedClass);
 }
 
