- Removed some "Win32" string references

diff --git a/Engine/Source/Developer/BlueprintNativeCodeGen/Private/BlueprintNativeCodeGenUtils.cpp b/Engine/Source/Developer/BlueprintNativeCodeGen/Private/BlueprintNativeCodeGenUtils.cpp
--- a/Engine/Source/Developer/BlueprintNativeCodeGen/Private/BlueprintNativeCodeGenUtils.cpp
+++ b/Engine/Source/Developer/BlueprintNativeCodeGen/Private/BlueprintNativeCodeGenUtils.cpp
@@ -133,8 +133,7 @@ static bool BlueprintNativeCodeGenUtilsImpl::GeneratePluginDescFile(const FBluep
 
 				FName UBTPlatformName = PlatformInfo->DataDrivenPlatformInfo->UBTPlatformName;
 				// Hack to allow clients for PS4/XboxOne (etc.) to build the nativized assets plugin
-				const bool bIsClientValidForPlatform = UBTPlatformName == TEXT("Win32") ||
-					UBTPlatformName == TEXT("Win64") ||
+				const bool bIsClientValidForPlatform = UBTPlatformName == TEXT("Win64") ||
 					UBTPlatformName == TEXT("Linux") ||
 					UBTPlatformName == TEXT("LinuxAArch64") ||
 					UBTPlatformName == TEXT("Mac");
diff --git a/Engine/Source/Developer/TargetPlatform/Private/TargetPlatformManagerModule.cpp b/Engine/Source/Developer/TargetPlatform/Private/TargetPlatformManagerModule.cpp
--- a/Engine/Source/Developer/TargetPlatform/Private/TargetPlatformManagerModule.cpp
+++ b/Engine/Source/Developer/TargetPlatform/Private/TargetPlatformManagerModule.cpp
@@ -1045,7 +1045,7 @@ protected:
 			{
 				PlatformInfo::EPlatformSDKStatus Status = PlatArray[Index+2].Contains(TEXT("INVALID")) ? PlatformInfo::EPlatformSDKStatus::NotInstalled : PlatformInfo::EPlatformSDKStatus::Installed;
 				FString PlatformName = PlatArray[Index+1];
-				if (PlatformName == TEXT("Win32") || PlatformName == TEXT("Win64"))
+				if (PlatformName == TEXT("Win64"))
 				{
 					PlatformName = TEXT("WindowsEditor");
 					PlatformInfo::UpdatePlatformSDKStatus(PlatformName, Status);
diff --git a/Engine/Source/Programs/AutomationTool/Scripts/AnalyzeThirdPartyLibs.Automation.cs b/Engine/Source/Programs/AutomationTool/Scripts/AnalyzeThirdPartyLibs.Automation.cs
--- a/Engine/Source/Programs/AutomationTool/Scripts/AnalyzeThirdPartyLibs.Automation.cs
+++ b/Engine/Source/Programs/AutomationTool/Scripts/AnalyzeThirdPartyLibs.Automation.cs
@@ -165,7 +165,7 @@ class AnalyzeThirdPartyLibs : BuildCommand
 
 		// Make a list of platforms
 		List<PlatformLibraryInfo> Platforms = new List<PlatformLibraryInfo>();
-		Platforms.Add(new PlatformLibraryInfo("Windows", "Windows", "Win32", "Win64", "VS20"));
+		Platforms.Add(new PlatformLibraryInfo("Windows", "Windows", "Win64", "VS20"));
 		Platforms.Add(new PlatformLibraryInfo("Mac", "Osx", "Mac"));
 		Platforms.Add(new PlatformLibraryInfo("iOS", "IOS"));
 		Platforms.Add(new PlatformLibraryInfo("Android", "Android"));
diff --git a/Engine/Source/Programs/AutomationTool/Scripts/RecordPerformance.Automation.cs b/Engine/Source/Programs/AutomationTool/Scripts/RecordPerformance.Automation.cs
--- a/Engine/Source/Programs/AutomationTool/Scripts/RecordPerformance.Automation.cs
+++ b/Engine/Source/Programs/AutomationTool/Scripts/RecordPerformance.Automation.cs
@@ -120,7 +120,7 @@ class RecordPerformance : BuildCommand
 
 	private IEnumerable<string> GetRhisForPlatform(string PlatformName)
 	{
-		if (PlatformName == "Win64" || PlatformName == "Win32")
+		if (PlatformName == "Win64")
 		{
 			return new string[] { "d3d11", "d3d12", "vulkan" };
 		}
